make sure the Gamethread Id is stored before configuring tests.

- fix for async loading thread assert on platforms that require cooked data.

#rnx

diff --git a/Engine/Source/Developer/LowLevelTestsRunner/Private/TestRunner.cpp b/Engine/Source/Developer/LowLevelTestsRunner/Private/TestRunner.cpp
--- a/Engine/Source/Developer/LowLevelTestsRunner/Private/TestRunner.cpp
+++ b/Engine/Source/Developer/LowLevelTestsRunner/Private/TestRunner.cpp
@@ -4,9 +4,14 @@
 
 #include "TestRunner.h"
 #include "TestHarness.h"
+#include "HAL/PlatformTLS.h"
 
 int RunTests(int argc, const char* argv[])
 {
+	// remember thread id of the main thread
+	GGameThreadId = FPlatformTLS::GetCurrentThreadId();
+	GIsGameThreadIdInitialized = true;
+
 	Setup();
 
 	//Read command-line from file (if any). Switch does it earlier.
