In preparation for the editor and enterprise low level tests:
- Added a CommonEngineInit file to share engine init code accross low level tests programs.
- Made the preprocessor definitions setup overridable per low level tests target.

diff --git a/Engine/Source/Developer/LowLevelTestsRunner/Public/CommonEngineInit.inl b/Engine/Source/Developer/LowLevelTestsRunner/Public/CommonEngineInit.inl
new file mode 100644
--- /dev/null
+++ b/Engine/Source/Developer/LowLevelTestsRunner/Public/CommonEngineInit.inl
@@ -0,0 +1,155 @@
+// Copyright Epic Games, Inc. All Rights Reserved.
+
+#pragma once
+
+#include "HAL/PlatformApplicationMisc.h"
+
+#include "RenderUtils.h"
+#include "ShaderParameterMetadata.h"
+
+#if WITH_EDITORONLY_DATA
+#include "DerivedDataBuild.h"
+#include "DerivedDataCache.h"
+#endif // WITH_EDITORONLY_DATA
+
+#include "DistanceFieldAtlas.h"
+#include "MeshCardRepresentation.h"
+
+/**
+ * A lot of what's in this file was taken from LaunchEngineLoop.cpp.
+ * It's made available here in an effort to reuse the same init code through various low level tests programs.
+ */
+
+void InitThreadPool()
+{
+	{
+		int32 NumThreadsInLargeThreadPool = FMath::Max( FPlatformMisc::NumberOfCoresIncludingHyperthreads() - 2, 2 );
+		int32 NumThreadsInThreadPool = FPlatformMisc::NumberOfWorkerThreadsToSpawn();
+		// when we are in the editor we like to do things like build lighting and such
+		// this thread pool can be used for those purposes
+		extern CORE_API int32 GUseNewTaskBackend;
+		if ( !GUseNewTaskBackend )
+		{
+			GLargeThreadPool = FQueuedThreadPool::Allocate();
+		}
+		else
+		{
+			GLargeThreadPool = new FQueuedLowLevelThreadPool();
+		}
+	
+		constexpr int32 StackSize = 128 * 1024;
+
+		// TaskGraph has it's HP threads slightly below normal, we want to be below the taskgraph HP threads to avoid interfering with the game-thread.
+		verify( GLargeThreadPool->Create( NumThreadsInLargeThreadPool, StackSize, TPri_BelowNormal, TEXT("LargeThreadPool") ) );
+
+		// GThreadPool will schedule on the LargeThreadPool but limit max concurrency to the given number.
+		GThreadPool = new FQueuedThreadPoolWrapper( GLargeThreadPool, NumThreadsInThreadPool );
+	}
+
+	if (FPlatformProcess::SupportsMultithreading())
+	{
+		GIOThreadPool = FQueuedThreadPool::Allocate();
+		int32 NumThreadsInThreadPool = FPlatformMisc::NumberOfIOWorkerThreadsToSpawn();
+		verify(GIOThreadPool->Create(NumThreadsInThreadPool, 96 * 1024, TPri_AboveNormal, TEXT("IOThreadPool")));
+	}
+}
+
+void InitTaskGraph()
+{
+	FTaskGraphInterface::Startup( FPlatformMisc::NumberOfWorkerThreadsToSpawn() );
+	FTaskGraphInterface::Get().AttachToThread( ENamedThreads::GameThread );
+}
+
+void InitOutputDevices()
+{
+#if WITH_APPLICATION_CORE
+	GError = FPlatformApplicationMisc::GetErrorOutputDevice();
+	GWarn = FPlatformApplicationMisc::GetFeedbackContext();
+#else
+	GError = FPlatformOutputDevices::GetError();
+	GWarn = FPlatformOutputDevices::GetFeedbackContext();
+#endif
+}
+
+#if STATS
+void InitStats()
+{
+	FThreadStats::StartThread();
+}
+#endif // #if STATS
+
+void InitRendering()
+{
+	FShaderParametersMetadata::InitializeAllUniformBufferStructs();
+
+	{
+		// Initialize the RHI.
+		const bool bHasEditorToken = false;
+		RHIInit(bHasEditorToken);
+	}
+
+	{
+		// One-time initialization of global variables based on engine configuration.
+		RenderUtilsInit();
+	}
+}
+
+#if WITH_EDITORONLY_DATA
+void InitDerivedDataCache()
+{
+	if (!FPlatformProperties::RequiresCookedData())
+	{
+		// Ensure that DDC is initialized from the game thread.
+		UE::DerivedData::GetBuild();
+		UE::DerivedData::GetCache();
+		GetDerivedDataCacheRef();
+	}
+}
+#endif // WITH_EDITORONLY_DATA
+
+void InitAsyncQueues()
+{
+	check(!GDistanceFieldAsyncQueue);
+	GDistanceFieldAsyncQueue = new FDistanceFieldAsyncQueue();
+
+	check(!GCardRepresentationAsyncQueue);
+	GCardRepresentationAsyncQueue = new FCardRepresentationAsyncQueue();
+}
+
+void InitSlate()
+{
+	FSlateApplication::Create();
+
+	TSharedPtr<FSlateRenderer> SlateRenderer = FModuleManager::Get().LoadModuleChecked<ISlateNullRendererModule>("SlateNullRenderer").CreateSlateNullRenderer();
+	TSharedRef<FSlateRenderer> SlateRendererSharedRef = SlateRenderer.ToSharedRef();
+
+	// If Slate is being used, initialize the renderer after RHIInit
+	FSlateApplication& CurrentSlateApp = FSlateApplication::Get();
+	CurrentSlateApp.InitializeRenderer(SlateRendererSharedRef);
+}
+
+void InitForWithEditorOnlyData()
+{
+#if WITH_COREUOBJECT
+	//Initialize the PackageResourceManager, which is needed to load any (non-script) Packages. It is first used in ProcessNewlyLoadedObjects (due to the loading of asset references in Class Default Objects)
+	// It has to be intialized after the AssetRegistryModule; the editor implementations of PackageResourceManager relies on it
+	IPackageResourceManager::Initialize();
+#endif
+#if WITH_EDITOR
+	// Initialize the BulkDataRegistry, which registers BulkData structs loaded from Packages for later building. It uses the same lifetime as IPackageResourceManager
+	IBulkDataRegistry::Initialize();
+#endif
+}
+
+#if WITH_EDITOR
+void InitEditor()
+{
+	FModuleManager::Get().LoadModuleChecked("UnrealEd");
+
+	GIsEditor = true;
+	GEngine = GEditor = NewObject<UEditorEngine>(GetTransientPackage(), UEditorEngine::StaticClass());
+
+	GEngine->ParseCommandline();
+	GEditor->InitEditor(&GEngineLoop);
+}
+#endif // #if WITH_EDITOR
diff --git a/Engine/Source/Programs/LowLevelTests/LowLevelTests.Target.cs b/Engine/Source/Programs/LowLevelTests/LowLevelTests.Target.cs
--- a/Engine/Source/Programs/LowLevelTests/LowLevelTests.Target.cs
+++ b/Engine/Source/Programs/LowLevelTests/LowLevelTests.Target.cs
@@ -48,6 +48,17 @@ public class LowLevelTestsTarget : TargetRules
 			bDebugBuildsActuallyUseDebugCRT = true;
 		}
 
+		SetupPreprocessorDefinitions(Target);
+	}
+
+	protected void SetupModule()
+	{
+		LaunchModuleName = this.GetType().Name.Replace("Target", string.Empty);
+		ExeBinariesSubFolder = LaunchModuleName;
+	}
+
+	protected virtual void SetupPreprocessorDefinitions(TargetInfo Target)
+	{
 		GlobalDefinitions.Add("STATS=0");
 
 		// Platform specific setup
@@ -72,10 +83,4 @@ public class LowLevelTestsTarget : TargetRules
 			bCompileAgainstApplicationCore = true;
 		}
 	}
-
-	protected void SetupModule()
-	{
-		LaunchModuleName = this.GetType().Name.Replace("Target", string.Empty);
-		ExeBinariesSubFolder = LaunchModuleName;
-	}
 }
