Add ability to LowLevelTestsRunner --global-base-moulde to Pre-load and unload modules outside of catch runs.
- This should avoid the startup.cpp and teardown.cpp pattern and more predictiablly tear-down and setup tests during GlobalSetup/Teardown.

diff --git a/Engine/Source/Developer/LowLevelTestsRunner/Private/TestRunner.cpp b/Engine/Source/Developer/LowLevelTestsRunner/Private/TestRunner.cpp
--- a/Engine/Source/Developer/LowLevelTestsRunner/Private/TestRunner.cpp
+++ b/Engine/Source/Developer/LowLevelTestsRunner/Private/TestRunner.cpp
@@ -4,6 +4,7 @@
 
 #include "TestRunner.h"
 #include "TestHarness.h"
+#include "Modules/ModuleManager.h"
 #include "HAL/PlatformTLS.h"
 
 #include "CommonEngineInit.inl"
@@ -58,6 +59,22 @@ void GlobalTeardown()
 	CleanupPlatform();
 }
 
+void LoadBaseTestModule(FString BaseModuleName)
+{
+	if (!BaseModuleName.IsEmpty())
+	{
+		FModuleManager::Get().LoadModule(*BaseModuleName);
+	}
+}
+
+void UnloadBaseTestModule(FString BaseModuleName)
+{
+	if (!BaseModuleName.IsEmpty())
+	{
+		FModuleManager::Get().UnloadModule(*BaseModuleName);
+	}
+}
+
 
 int RunTests(int argc, const char* argv[])
 {
@@ -97,6 +114,9 @@ int RunTests(int argc, const char* argv[])
 	// By default don't wait for input
 	bool bWaitForInputToTerminate = false;
 
+	
+	FString LoadModuleName = TEXT("");
+
 	for (int i = CatchArgc; i < argc; ++i)
 	{
 		if (std::strcmp(argv[i], "--wait") == 0)
@@ -119,6 +139,13 @@ int RunTests(int argc, const char* argv[])
 		{
 			bDebug = true;
 		}
+		// If we have --base-global-module parse proceeding argument as its option-value
+		if (std::strcmp(argv[i], "--base-global-module") == 0 &&
+			i + 1 < argc)
+		{
+			LoadModuleName = argv[i + 1];
+			++i;
+		}
 	}
 
 // Global initialization
@@ -126,6 +153,8 @@ int RunTests(int argc, const char* argv[])
 	GlobalSetup();
 #endif
 
+	LoadBaseTestModule(LoadModuleName);
+
 	int SessionResult = 0;
 	{
 		TGuardValue<bool> CatchRunning(bCatchIsRunning, true);
@@ -133,6 +162,8 @@ int RunTests(int argc, const char* argv[])
 		CatchArgv.Reset();
 	}
 
+	UnloadBaseTestModule(LoadModuleName);
+
 // Global cleanup
 #if defined(USE_GLOBAL_ENGINE_SETUP)
 	GlobalTeardown();
