LowLevelTests: Minimized code included by TestHarness.h

- Replaced CoreMinimal.h and StringConv.h with CoreTypes.h and forward declaration headers.
- Added push/pop for disabled warnings.
- Added FStringView string conversion.

#rnx

diff --git a/Engine/Source/Developer/LowLevelTestsRunner/Private/TestHarness.cpp b/Engine/Source/Developer/LowLevelTestsRunner/Private/TestHarness.cpp
new file mode 100644
--- /dev/null
+++ b/Engine/Source/Developer/LowLevelTestsRunner/Private/TestHarness.cpp
@@ -0,0 +1,35 @@
+// Copyright Epic Games, Inc. All Rights Reserved.
+
+#include "TestHarness.h"
+
+#include "Containers/StringConv.h"
+#include "Containers/StringView.h"
+#include "Containers/UnrealString.h"
+#include "Misc/StringBuilder.h"
+
+namespace UE::Private::LowLevelTestsRunner
+{
+
+std::string CStringToStdString(const TCHAR* Value)
+{
+	return StringViewToStdString(Value);
+}
+
+std::string FStringToStdString(const FString& Value)
+{
+	return StringViewToStdString(Value);
+}
+
+std::string StringViewToStdString(const FStringView& Value)
+{
+	const int32 ConvertedLen = FTCHARToUTF8_Convert::ConvertedLength(Value.GetData(), Value.Len());
+	std::string String;
+	String.reserve(ConvertedLen + 2);
+	String.push_back('"');
+	String.resize(ConvertedLen + 1);
+	FTCHARToUTF8_Convert::Convert(String.data() + 1, ConvertedLen, Value.GetData(), Value.Len());
+	String.push_back('"');
+	return String;
+}
+
+} // UE::Private::LowLevelTestsRunner
diff --git a/Engine/Source/Developer/LowLevelTestsRunner/Public/TestHarness.h b/Engine/Source/Developer/LowLevelTestsRunner/Public/TestHarness.h
--- a/Engine/Source/Developer/LowLevelTestsRunner/Public/TestHarness.h
+++ b/Engine/Source/Developer/LowLevelTestsRunner/Public/TestHarness.h
@@ -7,8 +7,9 @@
 #endif
 
 #if UE5_ENABLE_TESTHARNESS_ENGINE_SUPPORT
-#include "CoreMinimal.h"
-#include "Containers/StringConv.h"
+#include "CoreTypes.h"
+#include "Containers/ContainersFwd.h"
+#include "Containers/StringFwd.h"
 #endif
 
 #if defined(THIRD_PARTY_INCLUDES_START) && defined(THIRD_PARTY_INCLUDES_END)
@@ -16,93 +17,98 @@ THIRD_PARTY_INCLUDES_START
 #endif // defined(THIRD_PARTY_INCLUDES_START) && defined(THIRD_PARTY_INCLUDES_END)
 
 #ifdef _MSC_VER
-#pragma warning(disable: 4005 4582 4583) // 4005 is a macro redefinition. It's for "TEXT" and harmless in this case. 4582 and 4583 shouldn't be enabled by the engine in the first place.
 #pragma pack(push, 8)
+#pragma warning(push)
+#pragma warning(disable: 4005) // 'identifier': macro redefinition
+#pragma warning(disable: 4582) // 'type': constructor is not implicitly called
+#pragma warning(disable: 4583) // 'type': destructor is not implicitly called
 #endif
 #include "catch.hpp"
 #ifdef _MSC_VER
+#pragma warning(pop)
 #pragma pack(pop)
 #endif
 
-#include <stdio.h>
-#include <atomic>
-
 #if defined(THIRD_PARTY_INCLUDES_START) && defined(THIRD_PARTY_INCLUDES_END)
 THIRD_PARTY_INCLUDES_END
 #endif // defined(THIRD_PARTY_INCLUDES_START) && defined(THIRD_PARTY_INCLUDES_END)
 
-
 #if UE5_ENABLE_TESTHARNESS_ENGINE_SUPPORT
 
-inline std::string FStringToStdString(const FString& Value)
-{
-	static std::string quote(R"(")");
-	return quote + std::string(TCHAR_TO_UTF8(*Value)) + quote;
-}
+enum class ESPMode : uint8;
+template<class ObjectType, ESPMode InMode> class TSharedRef;
+template<class ObjectType, ESPMode InMode> class TSharedPtr;
+template <typename... Types> struct TTuple;
 
-inline FDateTime TestDateTime(const TCHAR* DateString)
+namespace UE::Private::LowLevelTestsRunner
 {
-	FDateTime Output;
-	REQUIRE(FDateTime::ParseIso8601(DateString, Output));
-	return Output;
+	std::string CStringToStdString(const TCHAR* Value);
+	std::string FStringToStdString(const FString& Value);
+	std::string StringViewToStdString(const FStringView& Value);
 }
 
-
-namespace Catch
+// Tell Catch how to print FString
+template <>
+struct Catch::StringMaker<FString>
 {
-	// Tell Catch how to print FStrings
-	template <>
-	class StringMaker<FString>
+	static std::string convert(const FString& Value)
 	{
-	public:
-		static std::string convert(const FString& Value) { return FStringToStdString(Value);  }
-	};
+		return UE::Private::LowLevelTestsRunner::FStringToStdString(Value);
+	}
+};
 
-	// Tell Catch how to print TCHAR strings
-	template <int SZ>
-	class StringMaker<TCHAR[SZ]>
+// Tell Catch how to print FStringView
+template <>
+struct Catch::StringMaker<FStringView>
+{
+	static std::string convert(const FStringView& Value)
 	{
-	public:
-		static std::string convert(const TCHAR* Value) { return FStringToStdString(FString((int32)SZ, Value)); }
-	};
+		return UE::Private::LowLevelTestsRunner::StringViewToStdString(Value);
+	}
+};
 
-	// Tell Catch how to print TPairs
-	template <typename PairKeyType, typename PairValueType>
-	class StringMaker<TPair<PairKeyType, PairValueType>>
+// Tell Catch how to print TCHAR strings
+template <int SZ>
+struct Catch::StringMaker<TCHAR[SZ]>
+{
+	static std::string convert(const TCHAR* Value)
 	{
-	public:
-		static std::string convert(const TPair<PairKeyType, PairValueType>& Value)
-		{
-			return "{ " + StringMaker<PairKeyType>::convert(Value.Key) + " , " + StringMaker<PairValueType>::convert(Value.Value) + " }";
-		}
-	};
+		return UE::Private::LowLevelTestsRunner::CStringToStdString(Value);
+	}
+};
 
-	// Tell Catch how to print TSharedRef
-	template<typename T, ESPMode Mode>
-	class StringMaker<TSharedRef<T, Mode>>
+// Tell Catch how to print TTuple<KeyType, ValueType>
+template <typename KeyType, typename ValueType>
+struct Catch::StringMaker<TTuple<KeyType, ValueType>>
+{
+	static std::string convert(const TTuple<KeyType, ValueType>& Value)
 	{
-	public:
-		static std::string convert(const TSharedRef<T, Mode>& Value)
-		{
-			return "0x" + std::string(TCHAR_TO_ANSI(*FString::Printf(TEXT("%p"), &Value.Get())));
-		}
-	};
+		return "{ " + StringMaker<KeyType>::convert(Value.Key) + " , " + StringMaker<ValueType>::convert(Value.Value) + " }";
+	}
+};
 
-	// Tell Catch how to print TSharedPtr
-	template<typename T, ESPMode Mode>
-	class StringMaker<TSharedPtr<T, Mode>>
+// Tell Catch how to print TSharedRef
+template<typename T, ESPMode Mode>
+struct Catch::StringMaker<TSharedRef<T, Mode>>
+{
+	static std::string convert(const TSharedRef<T, Mode>& Value)
 	{
-	public:
-		static std::string convert(const TSharedPtr<T, Mode>& Value)
-		{
-			return "0x" + std::string(TCHAR_TO_ANSI(*FString::Printf(TEXT("%p"), Value.Get())));
-		}
-	};
-}
+		return StringMaker<const T*>::convert(&Value.Get());
+	}
+};
 
+// Tell Catch how to print TSharedPtr
+template<typename T, ESPMode Mode>
+struct Catch::StringMaker<TSharedPtr<T, Mode>>
+{
+	static std::string convert(const TSharedPtr<T, Mode>& Value)
+	{
+		return StringMaker<const T*>::convert(Value.Get());
+	}
+};
 
 template <typename KeyT, typename ValueT>
-static bool operator==(const TMap<KeyT, ValueT>& Left, const TMap<KeyT, ValueT>& Right)
+inline bool operator==(const TMap<KeyT, ValueT>& Left, const TMap<KeyT, ValueT>& Right)
 {
 	bool bIsEqual = Left.Num() == Right.Num();
 	if (bIsEqual)
@@ -116,4 +122,4 @@ static bool operator==(const TMap<KeyT, ValueT>& Left, const TMap<KeyT, ValueT>&
 	return bIsEqual;
 }
 
-#endif // #if UE5_ENABLE_TESTHARNESS_ENGINE_SUPPORT
\ No newline at end of file
+#endif // #if UE5_ENABLE_TESTHARNESS_ENGINE_SUPPORT
