LowLevelTests: Changed string conversions to use std::ostream

#rnx

diff --git a/Engine/Source/Developer/LowLevelTestsRunner/Private/TestHarness.cpp b/Engine/Source/Developer/LowLevelTestsRunner/Private/TestHarness.cpp
--- a/Engine/Source/Developer/LowLevelTestsRunner/Private/TestHarness.cpp
+++ b/Engine/Source/Developer/LowLevelTestsRunner/Private/TestHarness.cpp
@@ -7,29 +7,42 @@
 #include "Containers/UnrealString.h"
 #include "Misc/StringBuilder.h"
 
-namespace UE::Private::LowLevelTestsRunner
+std::ostream& operator<<(std::ostream& Stream, const TCHAR* Value)
 {
+	return Stream << FUtf8StringView(FTCHARToUTF8(Value));
+}
+
+std::ostream& operator<<(std::ostream& Stream, const FString& Value)
+{
+	return Stream << FUtf8StringView(FTCHARToUTF8(Value));
+}
+
+std::ostream& operator<<(std::ostream& Stream, const FAnsiStringView& Value)
+{
+	return Stream.write(Value.GetData(), Value.Len());
+}
+
+std::ostream& operator<<(std::ostream& Stream, const FWideStringView& Value)
+{
+	return Stream << FUtf8StringView(FTCHARToUTF8(Value));
+}
 
-std::string CStringToStdString(const TCHAR* Value)
+std::ostream& operator<<(std::ostream& Stream, const FUtf8StringView& Value)
 {
-	return StringViewToStdString(Value);
+	return Stream.write((const char*)Value.GetData(), Value.Len());
 }
 
-std::string FStringToStdString(const FString& Value)
+std::ostream& operator<<(std::ostream& Stream, const FAnsiStringBuilderBase& Value)
 {
-	return StringViewToStdString(Value);
+	return Stream.write(Value.GetData(), Value.Len());
 }
 
-std::string StringViewToStdString(const FStringView& Value)
+std::ostream& operator<<(std::ostream& Stream, const FWideStringBuilderBase& Value)
 {
-	const int32 ConvertedLen = FTCHARToUTF8_Convert::ConvertedLength(Value.GetData(), Value.Len());
-	std::string String;
-	String.reserve(ConvertedLen + 2);
-	String.push_back('"');
-	String.resize(ConvertedLen + 1);
-	FTCHARToUTF8_Convert::Convert(String.data() + 1, ConvertedLen, Value.GetData(), Value.Len());
-	String.push_back('"');
-	return String;
+	return Stream << FUtf8StringView(FTCHARToUTF8(Value));
 }
 
-} // UE::Private::LowLevelTestsRunner
+std::ostream& operator<<(std::ostream& Stream, const FUtf8StringBuilderBase& Value)
+{
+	return Stream << Value.ToView();
+}
diff --git a/Engine/Source/Developer/LowLevelTestsRunner/Public/TestHarness.h b/Engine/Source/Developer/LowLevelTestsRunner/Public/TestHarness.h
--- a/Engine/Source/Developer/LowLevelTestsRunner/Public/TestHarness.h
+++ b/Engine/Source/Developer/LowLevelTestsRunner/Public/TestHarness.h
@@ -10,102 +10,41 @@
 #include "CoreTypes.h"
 #include "Containers/ContainersFwd.h"
 #include "Containers/StringFwd.h"
-#endif
-
-#if defined(THIRD_PARTY_INCLUDES_START) && defined(THIRD_PARTY_INCLUDES_END)
-THIRD_PARTY_INCLUDES_START
-#endif // defined(THIRD_PARTY_INCLUDES_START) && defined(THIRD_PARTY_INCLUDES_END)
-
-#ifdef _MSC_VER
-#pragma pack(push, 8)
-#pragma warning(push)
-#pragma warning(disable: 4005) // 'identifier': macro redefinition
-#pragma warning(disable: 4582) // 'type': constructor is not implicitly called
-#pragma warning(disable: 4583) // 'type': destructor is not implicitly called
-#endif
-#include "catch.hpp"
-#ifdef _MSC_VER
-#pragma warning(pop)
-#pragma pack(pop)
-#endif
 
-#if defined(THIRD_PARTY_INCLUDES_START) && defined(THIRD_PARTY_INCLUDES_END)
-THIRD_PARTY_INCLUDES_END
-#endif // defined(THIRD_PARTY_INCLUDES_START) && defined(THIRD_PARTY_INCLUDES_END)
+#include <ostream>
 
-#if UE5_ENABLE_TESTHARNESS_ENGINE_SUPPORT
+std::ostream& operator<<(std::ostream& Stream, const TCHAR* Value);
+std::ostream& operator<<(std::ostream& Stream, const FString& Value);
+std::ostream& operator<<(std::ostream& Stream, const FAnsiStringView& Value);
+std::ostream& operator<<(std::ostream& Stream, const FWideStringView& Value);
+std::ostream& operator<<(std::ostream& Stream, const FUtf8StringView& Value);
+std::ostream& operator<<(std::ostream& Stream, const FAnsiStringBuilderBase& Value);
+std::ostream& operator<<(std::ostream& Stream, const FWideStringBuilderBase& Value);
+std::ostream& operator<<(std::ostream& Stream, const FUtf8StringBuilderBase& Value);
 
-enum class ESPMode : uint8;
-template<class ObjectType, ESPMode InMode> class TSharedRef;
-template<class ObjectType, ESPMode InMode> class TSharedPtr;
 template <typename... Types> struct TTuple;
 
-namespace UE::Private::LowLevelTestsRunner
+template <typename KeyType, typename ValueType>
+std::ostream& operator<<(std::ostream& Stream, const TTuple<KeyType, ValueType>& Value)
 {
-	std::string CStringToStdString(const TCHAR* Value);
-	std::string FStringToStdString(const FString& Value);
-	std::string StringViewToStdString(const FStringView& Value);
+	return Stream << "{ " << Value.Key << " , " << Value.Value << " }";
 }
 
-// Tell Catch how to print FString
-template <>
-struct Catch::StringMaker<FString>
-{
-	static std::string convert(const FString& Value)
-	{
-		return UE::Private::LowLevelTestsRunner::FStringToStdString(Value);
-	}
-};
-
-// Tell Catch how to print FStringView
-template <>
-struct Catch::StringMaker<FStringView>
-{
-	static std::string convert(const FStringView& Value)
-	{
-		return UE::Private::LowLevelTestsRunner::StringViewToStdString(Value);
-	}
-};
-
-// Tell Catch how to print TCHAR strings
-template <int SZ>
-struct Catch::StringMaker<TCHAR[SZ]>
-{
-	static std::string convert(const TCHAR* Value)
-	{
-		return UE::Private::LowLevelTestsRunner::CStringToStdString(Value);
-	}
-};
-
-// Tell Catch how to print TTuple<KeyType, ValueType>
-template <typename KeyType, typename ValueType>
-struct Catch::StringMaker<TTuple<KeyType, ValueType>>
-{
-	static std::string convert(const TTuple<KeyType, ValueType>& Value)
-	{
-		return "{ " + StringMaker<KeyType>::convert(Value.Key) + " , " + StringMaker<ValueType>::convert(Value.Value) + " }";
-	}
-};
+enum class ESPMode : uint8;
+template <class ObjectType, ESPMode InMode> class TSharedRef;
+template <class ObjectType, ESPMode InMode> class TSharedPtr;
 
-// Tell Catch how to print TSharedRef
-template<typename T, ESPMode Mode>
-struct Catch::StringMaker<TSharedRef<T, Mode>>
+template <typename ObjectType, ESPMode Mode>
+std::ostream& operator<<(std::ostream& Stream, const TSharedRef<ObjectType, Mode>& Value)
 {
-	static std::string convert(const TSharedRef<T, Mode>& Value)
-	{
-		return StringMaker<const T*>::convert(&Value.Get());
-	}
-};
+	return Stream << &Value.Get();
+}
 
-// Tell Catch how to print TSharedPtr
-template<typename T, ESPMode Mode>
-struct Catch::StringMaker<TSharedPtr<T, Mode>>
+template <typename ObjectType, ESPMode Mode>
+std::ostream& operator<<(std::ostream& Stream, const TSharedPtr<ObjectType, Mode>& Value)
 {
-	static std::string convert(const TSharedPtr<T, Mode>& Value)
-	{
-		return StringMaker<const T*>::convert(Value.Get());
-	}
-};
+	return Stream << Value.Get();
+}
 
 template <typename KeyT, typename ValueT>
 inline bool operator==(const TMap<KeyT, ValueT>& Left, const TMap<KeyT, ValueT>& Right)
@@ -122,4 +61,25 @@ inline bool operator==(const TMap<KeyT, ValueT>& Left, const TMap<KeyT, ValueT>&
 	return bIsEqual;
 }
 
-#endif // #if UE5_ENABLE_TESTHARNESS_ENGINE_SUPPORT
+#endif // UE5_ENABLE_TESTHARNESS_ENGINE_SUPPORT
+
+#if defined(THIRD_PARTY_INCLUDES_START) && defined(THIRD_PARTY_INCLUDES_END)
+THIRD_PARTY_INCLUDES_START
+#endif // defined(THIRD_PARTY_INCLUDES_START) && defined(THIRD_PARTY_INCLUDES_END)
+
+#ifdef _MSC_VER
+#pragma pack(push, 8)
+#pragma warning(push)
+#pragma warning(disable: 4005) // 'identifier': macro redefinition
+#pragma warning(disable: 4582) // 'type': constructor is not implicitly called
+#pragma warning(disable: 4583) // 'type': destructor is not implicitly called
+#endif
+#include "catch.hpp"
+#ifdef _MSC_VER
+#pragma warning(pop)
+#pragma pack(pop)
+#endif
+
+#if defined(THIRD_PARTY_INCLUDES_START) && defined(THIRD_PARTY_INCLUDES_END)
+THIRD_PARTY_INCLUDES_END
+#endif // defined(THIRD_PARTY_INCLUDES_START) && defined(THIRD_PARTY_INCLUDES_END)
