LowLevelTests: Changed TTuple string conversion back to Catch::StringMaker

#rnx

diff --git a/Engine/Source/Developer/LowLevelTestsRunner/Public/TestHarness.h b/Engine/Source/Developer/LowLevelTestsRunner/Public/TestHarness.h
--- a/Engine/Source/Developer/LowLevelTestsRunner/Public/TestHarness.h
+++ b/Engine/Source/Developer/LowLevelTestsRunner/Public/TestHarness.h
@@ -22,14 +22,6 @@ std::ostream& operator<<(std::ostream& Stream, const FAnsiStringBuilderBase& Val
 std::ostream& operator<<(std::ostream& Stream, const FWideStringBuilderBase& Value);
 std::ostream& operator<<(std::ostream& Stream, const FUtf8StringBuilderBase& Value);
 
-template <typename... Types> struct TTuple;
-
-template <typename KeyType, typename ValueType>
-std::ostream& operator<<(std::ostream& Stream, const TTuple<KeyType, ValueType>& Value)
-{
-	return Stream << "{ " << Value.Key << " , " << Value.Value << " }";
-}
-
 enum class ESPMode : uint8;
 template <class ObjectType, ESPMode InMode> class TSharedRef;
 template <class ObjectType, ESPMode InMode> class TSharedPtr;
@@ -83,3 +75,19 @@ THIRD_PARTY_INCLUDES_START
 #if defined(THIRD_PARTY_INCLUDES_START) && defined(THIRD_PARTY_INCLUDES_END)
 THIRD_PARTY_INCLUDES_END
 #endif // defined(THIRD_PARTY_INCLUDES_START) && defined(THIRD_PARTY_INCLUDES_END)
+
+#if UE5_ENABLE_TESTHARNESS_ENGINE_SUPPORT
+
+// Tell Catch how to print TTuple<KeyType, ValueType>
+template <typename... Types> struct TTuple;
+
+template <typename KeyType, typename ValueType>
+struct Catch::StringMaker<TTuple<KeyType, ValueType>>
+{
+	static std::string convert(const TTuple<KeyType, ValueType>& Value)
+	{
+		return "{ " + StringMaker<KeyType>::convert(Value.Key) + " , " + StringMaker<ValueType>::convert(Value.Value) + " }";
+	}
+};
+
+#endif // #if UE5_ENABLE_TESTHARNESS_ENGINE_SUPPORT
