Adds ILowLevelTestsModule to LowLevelTestRuner
- Adds method to specify Global Initalization of LowLevelTest Modules
- Fixes math header issue in LowLevelTestModule on Mac

diff --git a/Engine/Source/Developer/LowLevelTestsRunner/Private/TestRunner.cpp b/Engine/Source/Developer/LowLevelTestsRunner/Private/TestRunner.cpp
--- a/Engine/Source/Developer/LowLevelTestsRunner/Private/TestRunner.cpp
+++ b/Engine/Source/Developer/LowLevelTestsRunner/Private/TestRunner.cpp
@@ -7,6 +7,7 @@
 #include "Modules/ModuleManager.h"
 #include "HAL/PlatformTLS.h"
 
+#include "LowLevelTestModule.h"
 #include "TestCommon/CoreUtilities.h"
 
 // Test run interceptor
@@ -39,6 +40,52 @@ void UnloadBaseTestModule(FString BaseModuleName)
 	}
 }
 
+void GlobalModuleSetup()
+{
+	// Search for all low level test modules
+	TArray<FName> ModuleNames;
+	FModuleManager::Get().FindModules(TEXT("*GlobalLowLevelTests"), ModuleNames);
+
+	if (ModuleNames.Num() <= 0)
+	{
+		return;
+	}
+
+	for (FName ModuleName : ModuleNames)
+	{
+		ILowLevelTestsModule* Module = FModuleManager::LoadModulePtr<ILowLevelTestsModule>(ModuleName);
+		if (Module != nullptr)
+		{
+			Module->GlobalSetup();
+		}
+	}
+}
+
+void GlobalModuleTeardown()
+{
+	// Search for all low level test modules
+	TArray<FName> ModuleNames;
+	FModuleManager::Get().FindModules(TEXT("*GlobalLowLevelTests"), ModuleNames);
+
+	if (ModuleNames.Num() <= 0)
+	{
+		return;
+	}
+
+	for (FName ModuleName : ModuleNames)
+	{
+		ILowLevelTestsModule* Module = FModuleManager::GetModulePtr<ILowLevelTestsModule>(ModuleName);
+		if (Module != nullptr)
+		{
+			Module->GlobalTeardown();
+			if (Module->SupportsAutomaticShutdown())
+			{
+				Module->ShutdownModule();
+			}
+		}
+	}
+}
+
 int RunTests(int argc, const char* argv[])
 {
 	// remember thread id of the main thread
@@ -112,6 +159,7 @@ int RunTests(int argc, const char* argv[])
 	}
 
 	LoadBaseTestModule(LoadModuleName);
+	GlobalModuleSetup();
 
 	int SessionResult = 0;
 	{
@@ -120,6 +168,7 @@ int RunTests(int argc, const char* argv[])
 		CatchArgv.Reset();
 	}
 
+	GlobalModuleTeardown();
 	UnloadBaseTestModule(LoadModuleName);
 
 	// Required, will crash on exit otherwise...
diff --git a/Engine/Source/Developer/LowLevelTestsRunner/Public/LowLevelTestModule.h b/Engine/Source/Developer/LowLevelTestsRunner/Public/LowLevelTestModule.h
new file mode 100644
--- /dev/null
+++ b/Engine/Source/Developer/LowLevelTestsRunner/Public/LowLevelTestModule.h
@@ -0,0 +1,13 @@
+// Copyright Epic Games, Inc. All Rights Reserved.
+
+#pragma once
+
+#include "Modules/ModuleManager.h"
+#include "Templates/SharedPointer.h"
+
+class ILowLevelTestsModule : public IModuleInterface
+{
+public:
+	virtual void GlobalSetup() = 0;
+	virtual void GlobalTeardown() = 0;
+};
diff --git a/Engine/Source/Developer/LowLevelTestsRunner/Public/TestHarness.h b/Engine/Source/Developer/LowLevelTestsRunner/Public/TestHarness.h
--- a/Engine/Source/Developer/LowLevelTestsRunner/Public/TestHarness.h
+++ b/Engine/Source/Developer/LowLevelTestsRunner/Public/TestHarness.h
@@ -11,6 +11,11 @@
 #include "Containers/ContainersFwd.h"
 #include "Containers/StringFwd.h"
 
+#if PLATFORM_MAC
+// Fix type redefinition error of FVector
+#include "HAL/PlatformMath.h"
+#endif // PLATFORM_MAC
+
 #include <ostream>
 
 std::ostream& operator<<(std::ostream& Stream, const TCHAR* Value);
