Catch unit tests: Redirect UE_LOG and ensure output to the global log singleton when running unit tests with Catch

Prior to this change running a catch executable (e.g., I tested with GeometryProcessingUnitTests.exe) there was no way to see any UE_LOGs or ensure output. Following this change ensure output is always shown, UE_LOG output with verbosity level <= Warning is shown if the command line argument "--log" is passed, and  UE_LOG output with verbosity level <= Log is shown if the command line arguement "--debug" is also passed. For even more verbose output the `DesiredLogVerbosity` variable in TestRunner.cpp can be edited.

#rnx

diff --git a/Engine/Source/Developer/LowLevelTestsRunner/Private/TestRunner.cpp b/Engine/Source/Developer/LowLevelTestsRunner/Private/TestRunner.cpp
--- a/Engine/Source/Developer/LowLevelTestsRunner/Private/TestRunner.cpp
+++ b/Engine/Source/Developer/LowLevelTestsRunner/Private/TestRunner.cpp
@@ -4,8 +4,11 @@
 
 #include "TestRunner.h"
 #include "TestHarness.h"
+#include "CoreGlobals.h"
 #include "Modules/ModuleManager.h"
 #include "HAL/PlatformTLS.h"
+#include "Misc/OutputDeviceRedirector.h"
+#include "Logging/LogVerbosity.h"
 
 #include "LowLevelTestModule.h"
 #include "TestCommon/CoreUtilities.h"
@@ -16,11 +19,45 @@
 namespace Catch
 {
 	// Test run interceptor
-	struct TestRunListener : public TestEventListenerBase {
+	struct TestRunListener : public TestEventListenerBase, public FOutputDevice
+	{
 		using TestEventListenerBase::TestEventListenerBase; // inherit constructor
 	private:
 		std::ostream& catchOut = getCurrentContext().getConfig()->stream();
 	public:
+		void testRunStarting(Catch::TestRunInfo const& testRunInfo) override {
+			// Register this event listener as an output device to enable reporting of UE_LOG, ensure etc
+			// Note: For UE_LOG(...) reporting the user must pass the "--log" command line argument, but this is not required for ensure reporting
+			GLog->AddOutputDevice(this);
+		}
+		
+		void testRunEnded(Catch::TestRunStats const& testRunStats) override {
+			GLog->RemoveOutputDevice(this);
+		}
+
+		// FOutputDevice interface
+		virtual void Serialize(const TCHAR* V, ELogVerbosity::Type LogVerbosity, const class FName& Category) override
+		{
+			// By default only log warnings/errors. If the user passes the command line argument "--debug" be more verbose
+			ELogVerbosity::Type DesiredLogVerbosity = ELogVerbosity::Warning;
+			if (bGDebug)
+			{
+				DesiredLogVerbosity = ELogVerbosity::Log; // Probably ELogVerbosity::Type::Verbose/VeryVerbose is too noisy
+			}
+
+			// TODO It might be nicer to increase the desired logging verbosity using the "-v high"/"--verbosity high"
+			// Catch option and changing condition above to `getCurrentContext().getConfig()->verbosity() == Verbosity::High`
+			// I tried this but unfortunately it didn't work, passing that option makes catch complain with the error message:
+			// "Verbosity level not supported by this reporter"
+			
+			// TODO Perhaps we should check IsInGameThread() or do something to make this threadsafe...?
+			if (LogVerbosity <= DesiredLogVerbosity)
+			{
+				catchOut << *FText::FromName(Category).ToString() << "(" << ToString(LogVerbosity) << ")" << ": " << V << "\n";
+			}
+		}
+		// End of FOutputDevice interface
+
 		void testCaseStarting(Catch::TestCaseInfo  const& TestInfo) override {
 			if (bGDebug)
 			{
