Introduce global setup for the LowLevelTests project.
Motivation: CoreUObject, taskgraph and threads are global constructs and need to be initialized as such and are expensive to setup/teardown before/after multiple single tests.

Introduced --global-setup and --no-global-setup to control execution of global test setup.
Unit tests that rely heavily on mocking and are designed to run without any global context should use --no-global-setup

Other: Removed CoreTestFixture and applied minor fixes wherever things didn't make sense.
One particular CoreUObject test was converted but had always failed to run on non-desktop platforms. It needs to be addressed ASAP.

diff --git a/Engine/Source/Developer/LowLevelTestsRunner/Private/TestCommon/CoreUObjectUtilities.cpp b/Engine/Source/Developer/LowLevelTestsRunner/Private/TestCommon/CoreUObjectUtilities.cpp
--- a/Engine/Source/Developer/LowLevelTestsRunner/Private/TestCommon/CoreUObjectUtilities.cpp
+++ b/Engine/Source/Developer/LowLevelTestsRunner/Private/TestCommon/CoreUObjectUtilities.cpp
@@ -3,6 +3,7 @@
 #if WITH_COREUOBJECT
 
 #include "TestCommon/CoreUObjectUtilities.h"
+
 #include "HAL/PlatformFileManager.h"
 #include "Misc/ConfigCacheIni.h"
 #include "Misc/CoreDelegates.h"
@@ -10,13 +11,8 @@
 #include "UObject/Package.h"
 #include "UObject/PackageResourceManager.h"
 
-bool CoreUObjectInitialised = false;
-
 void InitCoreUObject()
 {
-	if (CoreUObjectInitialised == true)
-		return;
-
 	//Initialize the PackageResourceManager, which is needed to load any (non-script) Packages. It is first used in ProcessNewlyLoadedObjects (due to the loading of asset references in Class Default Objects)
 	// It has to be intialized after the AssetRegistryModule; the editor implementations of PackageResourceManager relies on it
 	IPackageResourceManager::Initialize();
@@ -31,22 +27,15 @@ void InitCoreUObject()
 		FCoreDelegates::OnInit.Broadcast();
 		ProcessNewlyLoadedUObjects();
 	}
-
-	CoreUObjectInitialised = true;
 }
 
 void CleanupCoreUObject()
 {
-	if (CoreUObjectInitialised == false)
-		return;
-
 	IPackageResourceManager::Shutdown();
 
 	FCoreDelegates::OnPreExit.Broadcast();
 	MALLOC_PROFILER(GMalloc->Exec(nullptr, TEXT("MPROF STOP"), *GLog); );
 	FCoreDelegates::OnExit.Broadcast();
-
-	CoreUObjectInitialised = false;
 }
 
 #endif // WITH_COREUOBJECT
\ No newline at end of file
diff --git a/Engine/Source/Developer/LowLevelTestsRunner/Private/TestCommon/CoreUtilities.cpp b/Engine/Source/Developer/LowLevelTestsRunner/Private/TestCommon/CoreUtilities.cpp
--- a/Engine/Source/Developer/LowLevelTestsRunner/Private/TestCommon/CoreUtilities.cpp
+++ b/Engine/Source/Developer/LowLevelTestsRunner/Private/TestCommon/CoreUtilities.cpp
@@ -1,13 +1,15 @@
 // Copyright Epic Games, Inc. All Rights Reserved.
 
 #include "CoreMinimal.h"
-#include "Misc/CommandLine.h"
-#include "TestCommon/CoreUtilities.h"
+
 #include "HAL/PlatformOutputDevices.h"
 #include "HAL/PlatformProcess.h"
 #include "Logging/LogSuppressionInterface.h"
+#include "Misc/CommandLine.h"
 #include "Misc/QueuedThreadPool.h"
 
+#include "TestCommon/CoreUtilities.h"
+
 void InitCommandLine(bool AllowLogging)
 {
 	if (AllowLogging)
@@ -103,28 +105,16 @@ void CleanupTaskGraph()
 	FTaskGraphInterface::Shutdown();
 }
 
-bool TaskGraphInitialised = false;
-
 void InitTaskGraphAndDependencies(bool MultiThreaded)
 {
-	if (TaskGraphInitialised == true)
-		return;
-
 	InitTaskGraph(MultiThreaded);
 	InitAllThreadPools(MultiThreaded);
-
-	TaskGraphInitialised = true;
 }
 
 void CleanupTaskGraphAndDependencies()
 {
-	if (TaskGraphInitialised == false)
-		return;
-
 	CleanupAllThreadPools();
 	CleanupTaskGraph();
-
-	TaskGraphInitialised = false;
 }
 
 void CleanupPlatform()
diff --git a/Engine/Source/Developer/LowLevelTestsRunner/Private/TestFixtures/CoreTestFixture.cpp b/Engine/Source/Developer/LowLevelTestsRunner/Private/TestFixtures/CoreTestFixture.cpp
deleted file mode 100644
diff --git a/Engine/Source/Developer/LowLevelTestsRunner/Private/TestRunner.cpp b/Engine/Source/Developer/LowLevelTestsRunner/Private/TestRunner.cpp
--- a/Engine/Source/Developer/LowLevelTestsRunner/Private/TestRunner.cpp
+++ b/Engine/Source/Developer/LowLevelTestsRunner/Private/TestRunner.cpp
@@ -4,18 +4,25 @@
 
 #include "TestRunner.h"
 #include "TestHarness.h"
+
 #include "CoreGlobals.h"
-#include "Modules/ModuleManager.h"
 #include "HAL/PlatformTLS.h"
-#include "Misc/OutputDeviceRedirector.h"
 #include "Logging/LogVerbosity.h"
+#include "Misc/OutputDeviceRedirector.h"
+#include "Modules/ModuleManager.h"
+#include "Templates/SharedPointer.h"
 
 #include "LowLevelTestModule.h"
 #include "TestCommon/CoreUtilities.h"
-#include "TestCommon/CoreUObjectUtilities.h"
 
 #include <set>
 
+bool bCatchIsRunning = false;
+bool bGAllowLogging = false;
+bool bGMultithreaded = false;
+bool bGDebug = false;
+bool bGGlobalSetup = true;
+
 namespace Catch
 {
 	// Test run interceptor
@@ -30,7 +37,7 @@ namespace Catch
 			// Note: For UE_LOG(...) reporting the user must pass the "--log" command line argument, but this is not required for ensure reporting
 			GLog->AddOutputDevice(this);
 		}
-		
+
 		void testRunEnded(Catch::TestRunStats const& testRunStats) override {
 			GLog->RemoveOutputDevice(this);
 		}
@@ -49,13 +56,19 @@ namespace Catch
 			// Catch option and changing condition above to `getCurrentContext().getConfig()->verbosity() == Verbosity::High`
 			// I tried this but unfortunately it didn't work, passing that option makes catch complain with the error message:
 			// "Verbosity level not supported by this reporter"
-			
+
 			// TODO Perhaps we should check IsInGameThread() or do something to make this threadsafe...?
+			// UPDATE: CanBeUsedOnMultipleThreads returns true
 			if (LogVerbosity <= DesiredLogVerbosity)
 			{
 				catchOut << *FText::FromName(Category).ToString() << "(" << ToString(LogVerbosity) << ")" << ": " << V << "\n";
 			}
 		}
+
+		virtual bool CanBeUsedOnMultipleThreads() const
+		{
+			return true;
+		}
 		// End of FOutputDevice interface
 
 		void testCaseStarting(Catch::TestCaseInfo  const& TestInfo) override {
@@ -178,13 +191,15 @@ int RunTests(int argc, const char* argv[])
 	KnownLLTArgs.insert("--no-log");
 	KnownLLTArgs.insert("--mt");
 	KnownLLTArgs.insert("--no-mt");
-	KnownLLTArgs.insert("--debug");
+	KnownLLTArgs.insert("--global-setup");
+	KnownLLTArgs.insert("--no-global-setup");
 	KnownLLTArgs.insert("--base-global-module");
+	KnownLLTArgs.insert("--debug");
 
 	// By default don't wait for input
 	bool bWaitForInputToTerminate = false;
 
-	FString LoadModuleName = TEXT("");
+	FString BaseModuleName;
 
 	// Every argument that is not in the list of known low level test options and is considered to be a catch test runner argument.
 	for (int i = 0; i < argc; ++i)
@@ -203,7 +218,7 @@ int RunTests(int argc, const char* argv[])
 			{
 				bGAllowLogging = true;
 			}
-			if (std::strcmp(argv[i], "--no-log") == 0)
+			else if (std::strcmp(argv[i], "--no-log") == 0)
 			{
 				bGAllowLogging = false;
 			}
@@ -211,10 +226,18 @@ int RunTests(int argc, const char* argv[])
 			{
 				bGMultithreaded = true;
 			}
-			if (std::strcmp(argv[i], "--no-mt") == 0)
+			else if (std::strcmp(argv[i], "--no-mt") == 0)
 			{
 				bGMultithreaded = false;
 			}
+			if (std::strcmp(argv[i], "--global-setup") == 0)
+			{
+				bGGlobalSetup = true;
+			}
+			else if (std::strcmp(argv[i], "--no-global-setup") == 0)
+			{
+				bGGlobalSetup = false;
+			}
 			if (std::strcmp(argv[i], "--debug") == 0)
 			{
 				bGDebug = true;
@@ -223,7 +246,7 @@ int RunTests(int argc, const char* argv[])
 			if (std::strcmp(argv[i], "--base-global-module") == 0 &&
 				i + 1 < argc)
 			{
-				LoadModuleName = argv[i + 1];
+				BaseModuleName = argv[i + 1];
 				++i;
 			}
 		}
@@ -234,11 +257,14 @@ int RunTests(int argc, const char* argv[])
 		}
 	}
 
-	// Global command line initialization
-	InitCommandLine(bGAllowLogging);
+	if (bGGlobalSetup)
+	{	
+		InitCommandLine(bGAllowLogging);
+
+		LoadBaseTestModule(BaseModuleName);
 
-	LoadBaseTestModule(LoadModuleName);
-	GlobalModuleSetup();
+		GlobalModuleSetup();
+	}
 
 	int SessionResult = 0;
 	{
@@ -247,17 +273,16 @@ int RunTests(int argc, const char* argv[])
 		CatchArgv.Reset();
 	}
 
-	CleanupCoreUObject();
-	CleanupTaskGraphAndDependencies();
+	if (bGGlobalSetup)
+	{
+		GlobalModuleTeardown();
 
-	GlobalModuleTeardown();
-	UnloadBaseTestModule(LoadModuleName);
+		UnloadBaseTestModule(BaseModuleName);
 
-	// Required for platform cleanup, program will crash on exit otherwise
-	CleanupPlatform();
+		CleanupPlatform();
 
-	// Command line cleanup
-	CleanupCommandLine();
+		CleanupCommandLine();
+	}
 
 #if PLATFORM_DESKTOP
 	if (bWaitForInputToTerminate)
diff --git a/Engine/Source/Developer/LowLevelTestsRunner/Public/TestFixtures/CoreTestFixture.h b/Engine/Source/Developer/LowLevelTestsRunner/Public/TestFixtures/CoreTestFixture.h
deleted file mode 100644
diff --git a/Engine/Source/Developer/LowLevelTestsRunner/Public/TestHarness.h b/Engine/Source/Developer/LowLevelTestsRunner/Public/TestHarness.h
--- a/Engine/Source/Developer/LowLevelTestsRunner/Public/TestHarness.h
+++ b/Engine/Source/Developer/LowLevelTestsRunner/Public/TestHarness.h
@@ -98,11 +98,6 @@ struct Catch::StringMaker<TTuple<KeyType, ValueType>>
 	}
 };
 
-/** Use with TEST_CASE_METHOD to quickly migrate automation tests to Catch. */
-class FAutomationTestFixture
-{	
-};
-
 #define VERIFY(What, Actual)\
 	CAPTURE(What);\
 	CHECK(Actual == true)
diff --git a/Engine/Source/Developer/LowLevelTestsRunner/Public/TestRunner.h b/Engine/Source/Developer/LowLevelTestsRunner/Public/TestRunner.h
--- a/Engine/Source/Developer/LowLevelTestsRunner/Public/TestRunner.h
+++ b/Engine/Source/Developer/LowLevelTestsRunner/Public/TestRunner.h
@@ -20,10 +20,11 @@
 #pragma warning(disable : 4005)
 #endif
 
-static bool bCatchIsRunning = false;
-static bool bGAllowLogging = false;
-static bool bGMultithreaded = false;
-static bool bGDebug = false;
+extern bool bCatchIsRunning;
+extern bool bGAllowLogging;
+extern bool bGMultithreaded;
+extern bool bGDebug;
+extern bool bGGlobalSetup;
 
 #if __cplusplus
 	extern "C" const char* GetProcessExecutablePath();
diff --git a/Engine/Source/Programs/LowLevelTests/Tests/LowLevelTestsModule.cpp b/Engine/Source/Programs/LowLevelTests/Tests/LowLevelTestsModule.cpp
new file mode 100644
--- /dev/null
+++ b/Engine/Source/Programs/LowLevelTests/Tests/LowLevelTestsModule.cpp
@@ -0,0 +1,32 @@
+#include "TestCommon/CoreUtilities.h"
+#include "TestCommon/CoreUObjectUtilities.h"
+
+#include "LowLevelTestModule.h"
+#include "Modules/ModuleManager.h"
+
+class FGlobalLLTModule : public ILowLevelTestsModule
+{
+public:
+	virtual void GlobalSetup() override;
+	virtual void GlobalTeardown() override;
+};
+
+IMPLEMENT_MODULE(FGlobalLLTModule, GlobalLowLevelTests);
+
+void FGlobalLLTModule::GlobalSetup()
+{
+	InitTaskGraph();
+	InitThreadPool(true);
+#if WITH_COREUOBJECT
+	InitCoreUObject();
+#endif
+}
+
+void FGlobalLLTModule::GlobalTeardown()
+{
+#if WITH_COREUOBJECT
+	CleanupCoreUObject();
+#endif
+	CleanupThreadPool();
+	CleanupTaskGraph();
+}
\ No newline at end of file
diff --git a/Engine/Source/Runtime/CoreUObject/Tests/ObjectPtrTest.cpp b/Engine/Source/Runtime/CoreUObject/Tests/ObjectPtrTest.cpp
--- a/Engine/Source/Runtime/CoreUObject/Tests/ObjectPtrTest.cpp
+++ b/Engine/Source/Runtime/CoreUObject/Tests/ObjectPtrTest.cpp
@@ -158,43 +158,59 @@ TEST_CASE_METHOD(FObjectPtrTestBase, "CoreUObject::TObjectPtr::Null Behavior", "
 	TEST_TRUE(TEXT("Negation of a null object pointer should evaluate to true"), !NullObjectPtr);
 }
 
-TEST_CASE_METHOD(FObjectPtrTestBase, "CoreUObject::TObjectPtr::Default Serialize", "[CoreUObject][ObjectPtr][.Engine]")
+TEST_CASE_METHOD(FObjectPtrTestBase, "CoreUObject::TObjectPtr::Default Serialize", "[CoreUObject][ObjectPtr]")
 {
 	FSnapshotObjectRefMetrics ObjectRefMetrics(*this);
-	FObjectPtr DefaultTexturePtr(FObjectRef {FName("/Engine/EngineResources/DefaultTexture"), NAME_None, NAME_None, FObjectPathId("DefaultTexture")});
+
+	const FName TestPackageName(TEXT("/Engine/Test/ObjectPtrDefaultSerialize/Transient"));
+	UPackage* TestPackage = NewObject<UPackage>(nullptr, TestPackageName, RF_Transient);
+	TestPackage->AddToRoot();
+	UObject* TestSoftObject = NewObject<UTestDummyObject>(TestPackage, TEXT("DefaultSerializeObject"));
+
+	FObjectPtr DefaultSerializeObjectPtr(FObjectRef {FName("/Engine/Test/ObjectPtrDefaultSerialize/Transient"), NAME_None, NAME_None, FObjectPathId("DefaultSerializeObject")});
 
 	ObjectRefMetrics.TestNumResolves(TEXT("Unexpected resolve count after initializing an FObjectPtr"), UE_WITH_OBJECT_HANDLE_LATE_RESOLVE ? 0 : 1);
 	ObjectRefMetrics.TestNumFailedResolves(TEXT("Unexpected resolve failure after initializing an FObjectPtr"), 0);
 	ObjectRefMetrics.TestNumReads(TEXT("NumReads should not change when initializing an FObjectPtr"), 0);
 
 	FArchiveUObject Writer;
-	Writer << DefaultTexturePtr;
+	Writer << DefaultSerializeObjectPtr;
 
 	ObjectRefMetrics.TestNumResolves(TEXT("Serializing an FObjectPtr should force it to resolve"), 1);
 	ObjectRefMetrics.TestNumFailedResolves(TEXT("Unexpected resolve failure after serializing an FObjectPtr"), 0);
 	ObjectRefMetrics.TestNumReads(TEXT("NumReads should increase after serializing an FObjectPtr"), 1);
 
-	Writer << DefaultTexturePtr;
+	Writer << DefaultSerializeObjectPtr;
 
 	ObjectRefMetrics.TestNumResolves(TEXT("Serializing an FObjectPtr twice should only require it to resolve once"), 1);
 	ObjectRefMetrics.TestNumFailedResolves(TEXT("Unexpected resolve failure after serializing an FObjectPtr"), 0);
 	ObjectRefMetrics.TestNumReads(TEXT("NumReads should increase after serializing an FObjectPtr"), 2);
+
+	TestPackage->RemoveFromRoot();
 }
 
 TEST_CASE_METHOD(FObjectPtrTestBase, "CoreUObject::TObjectPtr::Soft Object Path", "[CoreUObject][ObjectPtr]")
 {
 	FSnapshotObjectRefMetrics ObjectRefMetrics(*this);
-	FObjectPtr DefaultTexturePtr(FObjectRef {FName("/Engine/EngineResources/DefaultTexture"), NAME_None, NAME_None, FObjectPathId("DefaultTexture")});
+
+	const FName TestPackageName(TEXT("/Engine/Test/ObjectPtrSoftObjectPath/Transient"));
+	UPackage* TestPackage = NewObject<UPackage>(nullptr, TestPackageName, RF_Transient);
+	TestPackage->AddToRoot();
+	UObject* TestSoftObject = NewObject<UTestDummyObject>(TestPackage, TEXT("TestSoftObject"));
+
+	FObjectPtr DefaultSoftObjPtr(FObjectRef {FName("/Engine/Test/ObjectPtrSoftObjectPath/Transient"), NAME_None, NAME_None, FObjectPathId("TestSoftObject")});
 
 	ObjectRefMetrics.TestNumResolves(TEXT("Unexpected resolve count after initializing an FObjectPtr"), UE_WITH_OBJECT_HANDLE_LATE_RESOLVE ? 0 : 1);
 	ObjectRefMetrics.TestNumFailedResolves(TEXT("Unexpected resolve failure after initializing an FObjectPtr"), 0);
 	ObjectRefMetrics.TestNumReads(TEXT("NumReads should not change when initializing an FObjectPtr"), 0);
 
-	FSoftObjectPath DefaultTexturePath(DefaultTexturePtr);
+	FSoftObjectPath DefaultSoftObjPath(DefaultSoftObjPtr);
+
+	ObjectRefMetrics.TestNumResolves(TEXT("Unexpected resolve count after initializing an FSoftObjectPath from an FObjectPtr"), 0);
 
-	ObjectRefMetrics.TestNumResolves(TEXT("Unexpected resolve count after initializing an FSoftObjectPath from an FObjectPtr"), UE_WITH_OBJECT_HANDLE_LATE_RESOLVE ? 0 : 1);
+	TEST_EQUAL_STR(TEXT("Soft object path constructed from an FObjectPtr does not have the expected path value"), TEXT("/Engine/Test/ObjectPtrSoftObjectPath/Transient.TestSoftObject"), *DefaultSoftObjPath.ToString());
 
-	TEST_EQUAL_STR(TEXT("Soft object path constructed from an FObjectPtr does not have the expected path value"), TEXT("/Engine/EngineResources/DefaultTexture.DefaultTexture"), *DefaultTexturePath.ToString());
+	TestPackage->RemoveFromRoot();
 }
 
 TEST_CASE_METHOD(FObjectPtrTestBase, "CoreUObject::TObjectPtr::Forward Declared", "[CoreUObject][ObjectPtr]")
diff --git a/Engine/Source/Runtime/CoreUObject/Tests/ObjectRefTrackingTestBase.cpp b/Engine/Source/Runtime/CoreUObject/Tests/ObjectRefTrackingTestBase.cpp
--- a/Engine/Source/Runtime/CoreUObject/Tests/ObjectRefTrackingTestBase.cpp
+++ b/Engine/Source/Runtime/CoreUObject/Tests/ObjectRefTrackingTestBase.cpp
@@ -2,14 +2,10 @@
 
 #include "ObjectRefTrackingTestBase.h"
 
-#if WITH_DEV_AUTOMATION_TESTS
-
 #if UE_WITH_OBJECT_HANDLE_TRACKING
 ObjectHandleReferenceResolvedFunction* FObjectRefTrackingTestBase::PrevResolvedFunc = nullptr;
 ObjectHandleReadFunction* FObjectRefTrackingTestBase::PrevReadFunc = nullptr;
 #endif
 thread_local uint32 FObjectRefTrackingTestBase::NumResolves = 0;
 thread_local uint32 FObjectRefTrackingTestBase::NumFailedResolves = 0;
-thread_local uint32 FObjectRefTrackingTestBase::NumReads = 0;
-
-#endif // WITH_DEV_AUTOMATION_TESTS
+thread_local uint32 FObjectRefTrackingTestBase::NumReads = 0;
\ No newline at end of file
diff --git a/Engine/Source/Runtime/CoreUObject/Tests/ObjectRefTrackingTestBase.h b/Engine/Source/Runtime/CoreUObject/Tests/ObjectRefTrackingTestBase.h
--- a/Engine/Source/Runtime/CoreUObject/Tests/ObjectRefTrackingTestBase.h
+++ b/Engine/Source/Runtime/CoreUObject/Tests/ObjectRefTrackingTestBase.h
@@ -2,14 +2,15 @@
 
 #pragma once
 
+#include "TestHarness.h"
+
 #include "CoreMinimal.h"
-#include "TestFixtures/CoreTestFixture.h"
 #include "UObject/ObjectHandle.h"
 
-class FObjectRefTrackingTestBase : public FCoreTestFixture
+
+class FObjectRefTrackingTestBase
 {
 public:
-
 	uint32 GetNumResolves() const { return NumResolves; }
 	uint32 GetNumFailedResolves() const { return NumFailedResolves; }
 	uint32 GetNumReads() const { return NumReads; }
diff --git a/Engine/Source/Runtime/CoreUObject/Tests/Serialization/BulkDataTests.cpp b/Engine/Source/Runtime/CoreUObject/Tests/Serialization/BulkDataTests.cpp
--- a/Engine/Source/Runtime/CoreUObject/Tests/Serialization/BulkDataTests.cpp
+++ b/Engine/Source/Runtime/CoreUObject/Tests/Serialization/BulkDataTests.cpp
@@ -6,7 +6,7 @@
 namespace BulkDataTest
 {
 	// Test code paths for BulkData objects that do not reference a file on disk.
-	TEST_CASE_METHOD(FAutomationTestFixture, "CoreUObject::Serialization::FByteBulkData::Transient", "[CoreUObject][Serialization]")
+	TEST_CASE("CoreUObject::Serialization::FByteBulkData::Transient", "[CoreUObject][Serialization]")
 	{
 		FByteBulkData BulkData;
 
