LLTs: Maximum log verbosity with --debug, allow logging from fault context

diff --git a/Engine/Source/Developer/LowLevelTestsRunner/Private/TestRunner.cpp b/Engine/Source/Developer/LowLevelTestsRunner/Private/TestRunner.cpp
--- a/Engine/Source/Developer/LowLevelTestsRunner/Private/TestRunner.cpp
+++ b/Engine/Source/Developer/LowLevelTestsRunner/Private/TestRunner.cpp
@@ -49,7 +49,7 @@ namespace Catch {
 			ELogVerbosity::Type DesiredLogVerbosity = ELogVerbosity::Warning;
 			if (bGDebug)
 			{
-				DesiredLogVerbosity = ELogVerbosity::Log; // Probably ELogVerbosity::Type::Verbose/VeryVerbose is too noisy
+				DesiredLogVerbosity = ELogVerbosity::VeryVerbose;
 			}
 
 			// TODO It might be nicer to increase the desired logging verbosity using the "-v high"/"--verbosity high"
@@ -69,6 +69,11 @@ namespace Catch {
 		{
 			return true;
 		}
+
+		virtual bool CanBeUsedOnPanicThread() const
+		{
+			return true;
+		}
 		// End of FOutputDevice interface
 
 		void testCaseStarting(TestCaseInfo  const& TestInfo) override {
