Catch2 - mini module cleanup

diff --git a/Engine/Source/Developer/LowLevelTestsRunner/LowLevelTestsRunner.Build.cs b/Engine/Source/Developer/LowLevelTestsRunner/LowLevelTestsRunner.Build.cs
--- a/Engine/Source/Developer/LowLevelTestsRunner/LowLevelTestsRunner.Build.cs
+++ b/Engine/Source/Developer/LowLevelTestsRunner/LowLevelTestsRunner.Build.cs
@@ -1,35 +1,38 @@
 // Copyright Epic Games, Inc. All Rights Reserved.
 
-using UnrealBuildTool;
 using System.IO;
-public class LowLevelTestsRunner : ModuleRules
+
+namespace UnrealBuildTool.Rules
 {
-	public LowLevelTestsRunner(ReadOnlyTargetRules Target) : base(Target)
+	public class LowLevelTestsRunner : ModuleRules
 	{
-		PCHUsage = PCHUsageMode.NoPCHs;
-		PrecompileForTargets = PrecompileTargetsType.None;
-		bUseUnity = false;
-		bRequiresImplementModule = false;
+		public LowLevelTestsRunner(ReadOnlyTargetRules Target) : base(Target)
+		{
+			PCHUsage = PCHUsageMode.NoPCHs;
+			PrecompileForTargets = PrecompileTargetsType.None;
+			bUseUnity = false;
+			bRequiresImplementModule = false;
 
-		PublicIncludePaths.AddRange(
-			new string[]
-			{
+			PublicIncludePaths.AddRange(
+				new string[]
+				{
 				"Runtime/Launch/Public",
 				Path.Combine(Target.UEThirdPartySourceDirectory, "Catch2", "v3.0.1", "src")
-			}
-		);
+				}
+			);
 
-		PrivateIncludePaths.AddRange(
-			new string[] {
+			PrivateIncludePaths.AddRange(
+				new string[] {
 				"Runtime/Launch/Private"
-			}
-		);
+				}
+			);
 
-		PrivateDependencyModuleNames.AddRange(
-			new string[] {
+			PrivateDependencyModuleNames.AddRange(
+				new string[] {
 				"Core",
 				"Projects"
-			}
-		);
+				}
+			);
+		}
 	}
 }
\ No newline at end of file
diff --git a/Engine/Source/Developer/LowLevelTestsRunner/Private/TestListeners/ConsoleListener.cpp b/Engine/Source/Developer/LowLevelTestsRunner/Private/TestListeners/ConsoleListener.cpp
new file mode 100644
--- /dev/null
+++ b/Engine/Source/Developer/LowLevelTestsRunner/Private/TestListeners/ConsoleListener.cpp
@@ -0,0 +1,73 @@
+// Copyright Epic Games, Inc. All Rights Reserved.
+
+#include "TestListeners/ConsoleListener.h"
+
+using namespace Catch;
+
+class ConsoleListener : public EventListenerBase, public FOutputDevice
+{
+	using EventListenerBase::EventListenerBase; // inherit constructor
+private:
+	void testRunStarting(TestRunInfo const& testRunInfo) override {
+		// Register this event listener as an output device to enable reporting of UE_LOG, ensure etc
+		// Note: For UE_LOG(...) reporting the user must pass the "--log" command line argument, but this is not required for ensure reporting
+		GLog->AddOutputDevice(this);
+	}
+
+	void testRunEnded(TestRunStats const& testRunStats) override {
+		GLog->RemoveOutputDevice(this);
+	}
+
+	// FOutputDevice interface
+	virtual void Serialize(const TCHAR* V, ELogVerbosity::Type LogVerbosity, const class FName& Category) override
+	{
+		// By default only log warnings/errors. If the user passes the command line argument "--debug" be more verbose
+		ELogVerbosity::Type DesiredLogVerbosity = ELogVerbosity::Warning;
+		if (bGDebug)
+		{
+			DesiredLogVerbosity = ELogVerbosity::VeryVerbose;
+		}
+
+		// TODO It might be nicer to increase the desired logging verbosity using the "-v high"/"--verbosity high"
+		// Catch option and changing condition above to `getCurrentContext().getConfig()->verbosity() == Verbosity::High`
+		// I tried this but unfortunately it didn't work, passing that option makes catch complain with the error message:
+		// "Verbosity level not supported by this reporter"
+
+		if (LogVerbosity <= DesiredLogVerbosity)
+		{
+			std::cout << *FText::FromName(Category).ToString() << "(" << ToString(LogVerbosity) << ")" << ": " << V << "\n";
+		}
+	}
+
+	virtual bool CanBeUsedOnMultipleThreads() const
+	{
+		return true;
+	}
+
+	virtual bool CanBeUsedOnPanicThread() const
+	{
+		return true;
+	}
+	// End of FOutputDevice interface
+
+	void testCaseStarting(TestCaseInfo  const& TestInfo) override {
+		if (bGDebug)
+		{
+			std::cout << TestInfo.lineInfo.file << ":" << TestInfo.lineInfo.line << " with tags " << TestInfo.tagsAsString() << " \n";
+		}
+	}
+
+	void testCaseEnded(TestCaseStats const& testCaseStats) override {
+		if (testCaseStats.totals.testCases.failed > 0)
+		{
+			std::cout << "* Error: Test case \"" << testCaseStats.testInfo->name << "\" failed \n";
+		}
+	}
+
+	void assertionEnded(AssertionStats const& assertionStats) override {
+		if (!assertionStats.assertionResult.succeeded()) {
+			std::cout << "* Error: Assertion \"" << assertionStats.assertionResult.getExpression() << "\" failed at " << assertionStats.assertionResult.getSourceInfo().file << ": " << assertionStats.assertionResult.getSourceInfo().line << "\n";
+		}
+	}
+};
+CATCH_REGISTER_LISTENER(ConsoleListener);
\ No newline at end of file
diff --git a/Engine/Source/Developer/LowLevelTestsRunner/Private/TestRunner.cpp b/Engine/Source/Developer/LowLevelTestsRunner/Private/TestRunner.cpp
--- a/Engine/Source/Developer/LowLevelTestsRunner/Private/TestRunner.cpp
+++ b/Engine/Source/Developer/LowLevelTestsRunner/Private/TestRunner.cpp
@@ -1,22 +1,17 @@
 // Copyright Epic Games, Inc. All Rights Reserved.
 
-#include "TestRunner.h"
-#include "TestHarness.h"
-
-#include <catch2/reporters/catch_reporter_event_listener.hpp>
-#include <catch2/catch_session.hpp>
-#include <catch2/reporters/catch_reporter_registrars.hpp>
-#include <catch2/catch_test_case_info.hpp>
-
-#include "CoreGlobals.h"
 #include "HAL/PlatformTLS.h"
-#include "Logging/LogVerbosity.h"
-#include "Misc/OutputDeviceRedirector.h"
 #include "Modules/ModuleManager.h"
 #include "Templates/SharedPointer.h"
 
+#include "TestRunner.h"
+#include "TestHarness.h"
+#include "TestGlobals.h"
 #include "LowLevelTestModule.h"
 #include "TestCommon/CoreUtilities.h"
+#include "TestListeners/ConsoleListener.h"
+
+#include <catch2/catch_session.hpp>
 
 #include <set>
 #include <iostream>
@@ -27,77 +22,6 @@ bool bGMultithreaded = false;
 bool bGDebug = false;
 bool bGGlobalSetup = true;
 
-namespace Catch {
-	class TestRunListener : public EventListenerBase, public FOutputDevice
-	{
-		using EventListenerBase::EventListenerBase; // inherit constructor
-	private:
-		void testRunStarting(TestRunInfo const& testRunInfo) override {
-			// Register this event listener as an output device to enable reporting of UE_LOG, ensure etc
-			// Note: For UE_LOG(...) reporting the user must pass the "--log" command line argument, but this is not required for ensure reporting
-			GLog->AddOutputDevice(this);
-		}
-
-		void testRunEnded(TestRunStats const& testRunStats) override {
-			GLog->RemoveOutputDevice(this);
-		}
-
-		// FOutputDevice interface
-		virtual void Serialize(const TCHAR* V, ELogVerbosity::Type LogVerbosity, const class FName& Category) override
-		{
-			// By default only log warnings/errors. If the user passes the command line argument "--debug" be more verbose
-			ELogVerbosity::Type DesiredLogVerbosity = ELogVerbosity::Warning;
-			if (bGDebug)
-			{
-				DesiredLogVerbosity = ELogVerbosity::VeryVerbose;
-			}
-
-			// TODO It might be nicer to increase the desired logging verbosity using the "-v high"/"--verbosity high"
-			// Catch option and changing condition above to `getCurrentContext().getConfig()->verbosity() == Verbosity::High`
-			// I tried this but unfortunately it didn't work, passing that option makes catch complain with the error message:
-			// "Verbosity level not supported by this reporter"
-
-			// TODO Perhaps we should check IsInGameThread() or do something to make this threadsafe...?
-			// UPDATE: CanBeUsedOnMultipleThreads returns true
-			if (LogVerbosity <= DesiredLogVerbosity)
-			{
-				std::cout << *FText::FromName(Category).ToString() << "(" << ToString(LogVerbosity) << ")" << ": " << V << "\n";
-			}
-		}
-
-		virtual bool CanBeUsedOnMultipleThreads() const
-		{
-			return true;
-		}
-
-		virtual bool CanBeUsedOnPanicThread() const
-		{
-			return true;
-		}
-		// End of FOutputDevice interface
-
-		void testCaseStarting(TestCaseInfo  const& TestInfo) override {
-			if (bGDebug)
-			{
-				std::cout << TestInfo.lineInfo.file << ":" << TestInfo.lineInfo.line << " with tags " << TestInfo.tagsAsString() << " \n";
-			}
-		}
-
-		void testCaseEnded(TestCaseStats const& testCaseStats) override {
-			if (testCaseStats.totals.testCases.failed > 0)
-			{
-				std::cout << "* Error: Test case \"" << testCaseStats.testInfo->name << "\" failed \n";
-			}
-		}
-
-		void assertionEnded(AssertionStats const& assertionStats) override {
-			if (!assertionStats.assertionResult.succeeded()) {
-				std::cout << "* Error: Assertion \"" << assertionStats.assertionResult.getExpression() << "\" failed at " << assertionStats.assertionResult.getSourceInfo().file << ": " << assertionStats.assertionResult.getSourceInfo().line << "\n";
-			}
-		}
-	};
-	CATCH_REGISTER_LISTENER(TestRunListener);
-}
 
 void LoadBaseTestModule(FString BaseModuleName)
 {
diff --git a/Engine/Source/Developer/LowLevelTestsRunner/Public/TestGlobals.h b/Engine/Source/Developer/LowLevelTestsRunner/Public/TestGlobals.h
new file mode 100644
--- /dev/null
+++ b/Engine/Source/Developer/LowLevelTestsRunner/Public/TestGlobals.h
@@ -0,0 +1,9 @@
+// Copyright Epic Games, Inc. All Rights Reserved.
+
+#pragma once
+
+extern bool bCatchIsRunning;
+extern bool bGAllowLogging;
+extern bool bGMultithreaded;
+extern bool bGDebug;
+extern bool bGGlobalSetup;
\ No newline at end of file
diff --git a/Engine/Source/Developer/LowLevelTestsRunner/Public/TestListeners/ConsoleListener.h b/Engine/Source/Developer/LowLevelTestsRunner/Public/TestListeners/ConsoleListener.h
new file mode 100644
--- /dev/null
+++ b/Engine/Source/Developer/LowLevelTestsRunner/Public/TestListeners/ConsoleListener.h
@@ -0,0 +1,16 @@
+// Copyright Epic Games, Inc. All Rights Reserved.
+
+#pragma once
+
+#include "CoreMinimal.h"
+#include "CoreGlobals.h"
+#include "Logging/LogVerbosity.h"
+#include "Misc/OutputDeviceRedirector.h"
+
+#include "TestGlobals.h"
+
+#include <catch2/catch_test_case_info.hpp>
+#include <catch2/reporters/catch_reporter_event_listener.hpp>
+#include <catch2/reporters/catch_reporter_registrars.hpp>
+
+#include <iostream>
\ No newline at end of file
diff --git a/Engine/Source/Developer/LowLevelTestsRunner/Public/TestRunner.h b/Engine/Source/Developer/LowLevelTestsRunner/Public/TestRunner.h
--- a/Engine/Source/Developer/LowLevelTestsRunner/Public/TestRunner.h
+++ b/Engine/Source/Developer/LowLevelTestsRunner/Public/TestRunner.h
@@ -20,12 +20,6 @@
 #pragma warning(disable : 4005)
 #endif
 
-extern bool bCatchIsRunning;
-extern bool bGAllowLogging;
-extern bool bGMultithreaded;
-extern bool bGDebug;
-extern bool bGGlobalSetup;
-
 #if __cplusplus
 	extern "C" const char* GetProcessExecutablePath();
 #else
diff --git a/Engine/Source/Developer/LowLevelTestsRunner/Public/catch_debugapi.h b/Engine/Source/Developer/LowLevelTestsRunner/Public/catch_debugapi.h
deleted file mode 100644
diff --git a/Engine/Source/Developer/LowLevelTestsRunner/Public/catch_windows.h b/Engine/Source/Developer/LowLevelTestsRunner/Public/catch_windows.h
deleted file mode 100644
