- Fix ConsoleListener attempting to send TCHAR* to cout, printing pointer values instead of strings.

diff --git a/Engine/Source/Developer/LowLevelTestsRunner/Private/TestListeners/ConsoleListener.cpp b/Engine/Source/Developer/LowLevelTestsRunner/Private/TestListeners/ConsoleListener.cpp
--- a/Engine/Source/Developer/LowLevelTestsRunner/Private/TestListeners/ConsoleListener.cpp
+++ b/Engine/Source/Developer/LowLevelTestsRunner/Private/TestListeners/ConsoleListener.cpp
@@ -35,7 +35,12 @@ private:
 
 		if (LogVerbosity <= DesiredLogVerbosity && !GLog->GetSuppressEventTag())
 		{
-			std::cout << *FText::FromName(Category).ToString() << "(" << ToString(LogVerbosity) << ")" << ": " << V << "\n";
+			std::cout << TCHAR_TO_UTF8(*FText::FromName(Category).ToString())
+				<< "("
+				<< TCHAR_TO_UTF8(ToString(LogVerbosity))
+				<< "): "
+				<< TCHAR_TO_UTF8(V)
+				<< "\n";
 		}
 	}
 
