[LowLevelTestsRunner] added `--break` to Catch command line args to break into debugger when attached by default.

diff --git a/Engine/Source/Developer/LowLevelTestsRunner/Private/TestRunner.cpp b/Engine/Source/Developer/LowLevelTestsRunner/Private/TestRunner.cpp
--- a/Engine/Source/Developer/LowLevelTestsRunner/Private/TestRunner.cpp
+++ b/Engine/Source/Developer/LowLevelTestsRunner/Private/TestRunner.cpp
@@ -108,7 +108,7 @@ int RunTests(int argc, const char* argv[])
 #endif
 
 	int CatchArgc = 0;
-	TUniquePtr<const char* []> CatchArgv = MakeUnique<const char* []>(argc);
+	TUniquePtr<const char* []> CatchArgv = MakeUnique<const char* []>(argc + 1);
 
 	std::set<std::string> KnownLLTArgs;
 	KnownLLTArgs.insert("--wait");
@@ -183,6 +183,9 @@ int RunTests(int argc, const char* argv[])
 		}
 	}
 
+	//add break by default to break in debugger on test failures when attached
+	CatchArgv.Get()[CatchArgc++] = "--break";
+
 	if (bGGlobalSetup)
 	{	
 		InitCommandLine(bGAllowLogging);
