CoreUObject: Disable the FObjectHandle::Resolve test until its warnings and errors are fixed

#rnx

diff --git a/Engine/Source/Developer/LowLevelTestsRunner/Public/TestHarness.h b/Engine/Source/Developer/LowLevelTestsRunner/Public/TestHarness.h
--- a/Engine/Source/Developer/LowLevelTestsRunner/Public/TestHarness.h
+++ b/Engine/Source/Developer/LowLevelTestsRunner/Public/TestHarness.h
@@ -82,8 +82,22 @@ THIRD_PARTY_INCLUDES_END
 #define LLT_CONCAT_(x, y) x##y
 #define LLT_CONCAT(x, y) LLT_CONCAT_(x,y)
 
+#define DISABLED_TEST_CASE_METHOD_INTERNAL(TestName, ClassName, ...) \
+	CATCH_INTERNAL_START_WARNINGS_SUPPRESSION \
+	CATCH_INTERNAL_SUPPRESS_GLOBALS_WARNINGS \
+	namespace \
+	{ \
+		struct TestName : INTERNAL_CATCH_REMOVE_PARENS(ClassName) \
+		{ \
+			void test(); \
+		}; \
+	} \
+	CATCH_INTERNAL_STOP_WARNINGS_SUPPRESSION \
+	void TestName::test()
+
 #define DISABLED_TEST_CASE(...) static void LLT_CONCAT(disabled_test_,__LINE__)()
-#define DISABLED_TEST_CASE_METHOD(...) static void LLT_CONCAT(disabled_test_method_,__LINE__)()
+#define DISABLED_TEST_CASE_METHOD(ClassName, ...) \
+	DISABLED_TEST_CASE_METHOD_INTERNAL(INTERNAL_CATCH_UNIQUE_NAME( CATCH2_INTERNAL_TEST_ ), ClassName, __VA_ARGS__)
 #define DISABLED_SCENARIO(...) static void LLT_CONCAT(disabled_scenario_,__LINE__)()
 #define DISABLED_SECTION(...) auto LLT_CONCAT(disabled_section_,__LINE__) = []()
 
diff --git a/Engine/Source/Runtime/CoreUObject/Tests/ObjectHandleTest.cpp b/Engine/Source/Runtime/CoreUObject/Tests/ObjectHandleTest.cpp
--- a/Engine/Source/Runtime/CoreUObject/Tests/ObjectHandleTest.cpp
+++ b/Engine/Source/Runtime/CoreUObject/Tests/ObjectHandleTest.cpp
@@ -175,7 +175,8 @@ TEST_CASE_METHOD(FObjectHandleTestBase, "CoreUObject::FObjectHandle::Resolve Eng
 	TestResolvableNonNull("/Engine/Test/ObjectPtrDefaultSerialize/Transient", "DefaultSerializeObject");
 }
 
-TEST_CASE_METHOD(FObjectHandleTestBase, "CoreUObject::FObjectHandle::Resolve Non Existent Target", "[CoreUObject][ObjectHandle]")
+// TODO: Disabled until warnings and errors related to loading a non-existent package have been fixed.
+DISABLED_TEST_CASE_METHOD(FObjectHandleTestBase, "CoreUObject::FObjectHandle::Resolve Non Existent Target", "[CoreUObject][ObjectHandle]")
 {
 	// Confirm we don't successfully resolve an incorrect reference to engine content
 	TestResolveFailure("/Engine/EngineResources/NonExistentPackageName_0", "DefaultTexture");
