LowLevelTests: Use Unreal logging functionality instead of LocalPrint

diff --git a/Engine/Source/Developer/LowLevelTestsRunner/Private/TestCommon/ApplicationCoreUtilities.cpp b/Engine/Source/Developer/LowLevelTestsRunner/Private/TestCommon/ApplicationCoreUtilities.cpp
deleted file mode 100644
diff --git a/Engine/Source/Developer/LowLevelTestsRunner/Private/TestCommon/CoreUtilities.cpp b/Engine/Source/Developer/LowLevelTestsRunner/Private/TestCommon/CoreUtilities.cpp
--- a/Engine/Source/Developer/LowLevelTestsRunner/Private/TestCommon/CoreUtilities.cpp
+++ b/Engine/Source/Developer/LowLevelTestsRunner/Private/TestCommon/CoreUtilities.cpp
@@ -10,28 +10,12 @@
 
 #include "TestCommon/CoreUtilities.h"
 
-void InitCommandLine(bool AllowLogging)
+void InitCommandLine()
 {
-	if (AllowLogging)
-	{
-		FCommandLine::Set(TEXT(""));
-	}
-	else
-	{
-		FCommandLine::Set(TEXT(R"(-LogCmds="global off")"));
-		FLogSuppressionInterface::Get().ProcessConfigAndCommandLine();
-	}
 }
 
 void CleanupCommandLine()
 {
-	FCommandLine::Reset();
-}
-
-void InitOutputDevices()
-{
-	GError = FPlatformOutputDevices::GetError();
-	GWarn = FPlatformOutputDevices::GetFeedbackContext();
 }
 
 void InitIOThreadPool(bool MultiThreaded, int32 StackSize)
diff --git a/Engine/Source/Developer/LowLevelTestsRunner/Private/TestCommon/Initialization.cpp b/Engine/Source/Developer/LowLevelTestsRunner/Private/TestCommon/Initialization.cpp
--- a/Engine/Source/Developer/LowLevelTestsRunner/Private/TestCommon/Initialization.cpp
+++ b/Engine/Source/Developer/LowLevelTestsRunner/Private/TestCommon/Initialization.cpp
@@ -1,7 +1,6 @@
 // Copyright Epic Games, Inc. All Rights Reserved.
 
 #include "TestCommon/Initialization.h"
-#include "TestCommon/ApplicationCoreUtilities.h"
 #include "TestCommon/CoreUtilities.h"
 #include "TestCommon/CoreUObjectUtilities.h"
 #include "TestCommon/EditorUtilities.h"
@@ -17,15 +16,6 @@ void InitAllThreadPoolsEditorEx(bool MultiThreaded)
 	InitAllThreadPools(MultiThreaded);
 }
 
-void InitOutputDevicesEx()
-{
-#if WITH_APPLICATION_CORE
-	InitOutputDevicesAppCore();
-#else
-	InitOutputDevices();
-#endif
-}
-
 void InitStats()
 {
 #if STATS
@@ -35,13 +25,11 @@ void InitStats()
 
 void InitAll(bool bAllowLogging, bool bMultithreaded)
 {
-	InitCommandLine(bAllowLogging);
 	InitAllThreadPools(bMultithreaded);
 #if WITH_ENGINE
 	InitAsyncQueues();
 #endif // WITH_ENGINE
 	InitTaskGraph();
-	InitOutputDevices();
 #if WITH_ENGINE
 	InitRendering();
 #endif // WITH_ENGINE
diff --git a/Engine/Source/Developer/LowLevelTestsRunner/Private/TestListeners/ConsoleListener.cpp b/Engine/Source/Developer/LowLevelTestsRunner/Private/TestListeners/ConsoleListener.cpp
deleted file mode 100644
diff --git a/Engine/Source/Developer/LowLevelTestsRunner/Private/TestListeners/EventListener.cpp b/Engine/Source/Developer/LowLevelTestsRunner/Private/TestListeners/EventListener.cpp
new file mode 100644
--- /dev/null
+++ b/Engine/Source/Developer/LowLevelTestsRunner/Private/TestListeners/EventListener.cpp
@@ -0,0 +1,61 @@
+// Copyright Epic Games, Inc. All Rights Reserved.
+
+#include "Logging/LogMacros.h"
+#include "TestRunnerPrivate.h"
+
+#include <catch2/catch_test_case_info.hpp>
+#include <catch2/reporters/catch_reporter_event_listener.hpp>
+#include <catch2/reporters/catch_reporter_registrars.hpp>
+#include <catch2/reporters/catch_reporter_streaming_base.hpp>
+
+namespace UE::LowLevelTests
+{
+
+DEFINE_LOG_CATEGORY_STATIC(LogLowLevelTests, Log, VeryVerbose);
+
+class FEventListener final : public Catch::EventListenerBase
+{
+public:
+	using EventListenerBase::EventListenerBase;
+
+	void testCaseStarting(const Catch::TestCaseInfo& TestInfo) final
+	{
+		if (ITestRunner* TestRunner = ITestRunner::Get())
+		{
+			UE_CLOG(TestRunner->IsDebugMode(), LogLowLevelTests, Log,
+				TEXT("Started test case \"%hs\" at %hs(%" SIZE_T_FMT ") with tags %hs"),
+				TestInfo.name.c_str(),
+				TestInfo.lineInfo.file,
+				SIZE_T(TestInfo.lineInfo.line),
+				TestInfo.tagsAsString().c_str());
+		}
+	}
+
+	void testCaseEnded(const Catch::TestCaseStats& TestCaseStats) final
+	{
+		if (ITestRunner* TestRunner = ITestRunner::Get())
+		{
+			UE_CLOG(TestRunner->IsDebugMode() && TestCaseStats.totals.testCases.failed > 0,
+				LogLowLevelTests, Error, TEXT("Failed test case \"%hs\" at %hs(%" SIZE_T_FMT ")"),
+				TestCaseStats.testInfo->name.c_str(),
+				TestCaseStats.testInfo->lineInfo.file,
+				SIZE_T(TestCaseStats.testInfo->lineInfo.line));
+		}
+	}
+
+	void assertionEnded(const Catch::AssertionStats& AssertionStats) final
+	{
+		if (ITestRunner* TestRunner = ITestRunner::Get())
+		{
+			UE_CLOG(TestRunner->HasLogOutput() && !AssertionStats.assertionResult.succeeded(),
+				LogLowLevelTests, Error, TEXT("Assertion \"%hs\" failed at %hs(%" SIZE_T_FMT ")"),
+				AssertionStats.assertionResult.getExpression().c_str(),
+				AssertionStats.assertionResult.getSourceInfo().file,
+				SIZE_T(AssertionStats.assertionResult.getSourceInfo().line));
+		}
+	}
+};
+
+CATCH_REGISTER_LISTENER(FEventListener);
+
+} // UE::LowLevelTests
diff --git a/Engine/Source/Developer/LowLevelTestsRunner/Private/TestRunner.cpp b/Engine/Source/Developer/LowLevelTestsRunner/Private/TestRunner.cpp
--- a/Engine/Source/Developer/LowLevelTestsRunner/Private/TestRunner.cpp
+++ b/Engine/Source/Developer/LowLevelTestsRunner/Private/TestRunner.cpp
@@ -1,29 +1,65 @@
 // Copyright Epic Games, Inc. All Rights Reserved.
 
+#include "TestRunner.h"
+#include "TestRunnerPrivate.h"
+
+#include "HAL/PlatformOutputDevices.h"
 #include "HAL/PlatformTLS.h"
+#include "LowLevelTestModule.h"
 #include "Modules/ModuleManager.h"
 #include "Templates/SharedPointer.h"
-
-#include "TestRunner.h"
-#include "TestHarness.h"
-#include "TestGlobals.h"
-#include "LowLevelTestModule.h"
 #include "TestCommon/CoreUtilities.h"
-#include "TestListeners/ConsoleListener.h"
+#include "TestHarness.h"
+
+#if WITH_APPLICATION_CORE
+#include "HAL/PlatformApplicationMisc.h"
+#endif
 
 #include <catch2/catch_session.hpp>
 
-#include <set>
 #include <iostream>
+#include <set>
+
+namespace UE::LowLevelTests
+{
 
-bool bCatchIsRunning = false;
-bool bGAllowLogging = false;
-bool bGMultithreaded = false;
-bool bGDebug = false;
-bool bGGlobalSetup = true;
+static ITestRunner* GTestRunner;
+
+ITestRunner* ITestRunner::Get()
+{
+	return GTestRunner;
+}
 
+ITestRunner::ITestRunner()
+{
+	check(!GTestRunner);
+	GTestRunner = this;
+}
+
+ITestRunner::~ITestRunner()
+{
+	check(GTestRunner == this);
+	GTestRunner = nullptr;
+}
 
-void LoadBaseTestModule(FString BaseModuleName)
+class FTestRunner final : public ITestRunner
+{
+public:
+	bool HasGlobalSetup() const final { return bGlobalSetup; }
+	bool HasLogOutput() const final { return bLogOutput || bDebugMode; }
+	bool IsDebugMode() const final { return bDebugMode; }
+
+	void SetGlobalSetup(bool bInGlobalSetup) { bGlobalSetup = bInGlobalSetup; }
+	void SetLogOutput(bool bInLogOutput) { bLogOutput = bInLogOutput; }
+	void SetDebugMode(bool bInDebugMode) { bDebugMode = bInDebugMode; }
+
+private:
+	bool bGlobalSetup = true;
+	bool bLogOutput = false;
+	bool bDebugMode = false;
+};
+
+static void LoadBaseTestModule(const FString& BaseModuleName)
 {
 	if (!BaseModuleName.IsEmpty())
 	{
@@ -31,7 +67,7 @@ void LoadBaseTestModule(FString BaseModuleName)
 	}
 }
 
-void UnloadBaseTestModule(FString BaseModuleName)
+static void UnloadBaseTestModule(const FString& BaseModuleName)
 {
 	if (!BaseModuleName.IsEmpty())
 	{
@@ -39,42 +75,29 @@ void UnloadBaseTestModule(FString BaseModuleName)
 	}
 }
 
-void GlobalModuleSetup()
+static TArray<FName> GetGlobalModuleNames()
 {
-	// Search for all low level test modules
 	TArray<FName> ModuleNames;
 	FModuleManager::Get().FindModules(TEXT("*GlobalLowLevelTests"), ModuleNames);
+	return ModuleNames;
+}
 
-	if (ModuleNames.Num() <= 0)
-	{
-		return;
-	}
-
-	for (FName ModuleName : ModuleNames)
+static void GlobalModuleSetup()
+{
+	for (FName ModuleName : GetGlobalModuleNames())
 	{
-		ILowLevelTestsModule* Module = FModuleManager::LoadModulePtr<ILowLevelTestsModule>(ModuleName);
-		if (Module != nullptr)
+		if (ILowLevelTestsModule* Module = FModuleManager::LoadModulePtr<ILowLevelTestsModule>(ModuleName))
 		{
 			Module->GlobalSetup();
 		}
 	}
 }
 
-void GlobalModuleTeardown()
+static void GlobalModuleTeardown()
 {
-	// Search for all low level test modules
-	TArray<FName> ModuleNames;
-	FModuleManager::Get().FindModules(TEXT("*GlobalLowLevelTests"), ModuleNames);
-
-	if (ModuleNames.Num() <= 0)
-	{
-		return;
-	}
-
-	for (FName ModuleName : ModuleNames)
+	for (FName ModuleName : GetGlobalModuleNames())
 	{
-		ILowLevelTestsModule* Module = FModuleManager::GetModulePtr<ILowLevelTestsModule>(ModuleName);
-		if (Module != nullptr)
+		if (ILowLevelTestsModule* Module = FModuleManager::GetModulePtr<ILowLevelTestsModule>(ModuleName))
 		{
 			Module->GlobalTeardown();
 			if (Module->SupportsAutomaticShutdown())
@@ -85,15 +108,19 @@ void GlobalModuleTeardown()
 	}
 }
 
+} // UE::LowLevelTests
+
 int RunTests(int argc, const char* argv[])
 {
-	// remember thread id of the main thread
+	using namespace UE::LowLevelTests;
+
+	// Set up the Game Thread.
 	GGameThreadId = FPlatformTLS::GetCurrentThreadId();
 	GIsGameThreadIdInitialized = true;
 
 #ifdef SLEEP_ON_INIT
-	// Sleep to allow sync with Gauntlet
-	std::this_thread::sleep_for(std::chrono::milliseconds(5000));
+	// Sleep to allow sync with Gauntlet.
+	FPlatformProcess::Sleep(5.0f);
 #endif
 
 	//Read command-line from file (if any). Some platforms do this earlier.
@@ -125,6 +152,9 @@ int RunTests(int argc, const char* argv[])
 	// By default don't wait for input
 	bool bWaitForInputToTerminate = false;
 
+	FTestRunner TestRunner;
+	bool bMultiThreaded = false;
+
 	FString BaseModuleName;
 
 	// Every argument that is not in the list of known low level test options and is considered to be a catch test runner argument.
@@ -142,35 +172,34 @@ int RunTests(int argc, const char* argv[])
 			}
 			if (std::strcmp(argv[i], "--log") == 0)
 			{
-				bGAllowLogging = true;
+				TestRunner.SetLogOutput(true);
 			}
 			else if (std::strcmp(argv[i], "--no-log") == 0)
 			{
-				bGAllowLogging = false;
+				TestRunner.SetLogOutput(false);
 			}
 			if (std::strcmp(argv[i], "--mt") == 0)
 			{
-				bGMultithreaded = true;
+				bMultiThreaded = true;
 			}
 			else if (std::strcmp(argv[i], "--no-mt") == 0)
 			{
-				bGMultithreaded = false;
+				bMultiThreaded = false;
 			}
 			if (std::strcmp(argv[i], "--global-setup") == 0)
 			{
-				bGGlobalSetup = true;
+				TestRunner.SetGlobalSetup(true);
 			}
 			else if (std::strcmp(argv[i], "--no-global-setup") == 0)
 			{
-				bGGlobalSetup = false;
+				TestRunner.SetGlobalSetup(false);
 			}
 			if (std::strcmp(argv[i], "--debug") == 0)
 			{
-				bGDebug = true;
+				TestRunner.SetDebugMode(true);
 			}
 			// If we have --base-global-module parse proceeding argument as its option-value
-			if (std::strcmp(argv[i], "--base-global-module") == 0 &&
-				i + 1 < argc)
+			if (std::strcmp(argv[i], "--base-global-module") == 0 && i + 1 < argc)
 			{
 				BaseModuleName = argv[i + 1];
 				++i;
@@ -183,12 +212,36 @@ int RunTests(int argc, const char* argv[])
 		}
 	}
 
-	//add break by default to break in debugger on test failures when attached
+	// Break in the debugger on failed assertions when attached.
 	CatchArgv.Get()[CatchArgc++] = "--break";
 
-	if (bGGlobalSetup)
-	{	
-		InitCommandLine(bGAllowLogging);
+	if (TestRunner.HasGlobalSetup())
+	{
+		FCommandLine::Set(TEXT(""));
+
+		// Finish setting up the Game Thread, which requires the command line.
+		FPlatformProcess::SetThreadAffinityMask(FPlatformAffinity::GetMainGameMask());
+		FPlatformProcess::SetupGameThread();
+
+		// Always set up GError to handle FatalError, failed assertions, and crashes and other fatal errors.
+	#if WITH_APPLICATION_CORE
+		GError = FPlatformApplicationMisc::GetErrorOutputDevice();
+	#else
+		GError = FPlatformOutputDevices::GetError();
+	#endif
+
+		if (TestRunner.HasLogOutput())
+		{
+			// Set up GWarn to handle Error, Warning, Display; but only when log output is enabled.
+		#if WITH_APPLICATION_CORE
+			GWarn = FPlatformApplicationMisc::GetFeedbackContext();
+		#else
+			GWarn = FPlatformOutputDevices::GetFeedbackContext();
+		#endif
+
+			// Set up default output devices to handle Log, Verbose, VeryVerbose.
+			FPlatformOutputDevices::SetupOutputDevices();
+		}
 
 		LoadBaseTestModule(BaseModuleName);
 
@@ -197,12 +250,11 @@ int RunTests(int argc, const char* argv[])
 
 	int SessionResult = 0;
 	{
-		TGuardValue<bool> CatchRunning(bCatchIsRunning, true);
 		SessionResult = Catch::Session().run(CatchArgc, CatchArgv.Get());
 		CatchArgv.Reset();
 	}
 
-	if (bGGlobalSetup)
+	if (TestRunner.HasGlobalSetup())
 	{
 		GlobalModuleTeardown();
 
@@ -210,7 +262,7 @@ int RunTests(int argc, const char* argv[])
 
 		CleanupPlatform();
 
-		CleanupCommandLine();
+		FCommandLine::Reset();
 	}
 
 #if PLATFORM_DESKTOP
@@ -222,4 +274,4 @@ int RunTests(int argc, const char* argv[])
 #endif
 
 	return SessionResult;
-};
\ No newline at end of file
+}
diff --git a/Engine/Source/Developer/LowLevelTestsRunner/Private/TestRunnerPrivate.h b/Engine/Source/Developer/LowLevelTestsRunner/Private/TestRunnerPrivate.h
new file mode 100644
--- /dev/null
+++ b/Engine/Source/Developer/LowLevelTestsRunner/Private/TestRunnerPrivate.h
@@ -0,0 +1,26 @@
+// Copyright Epic Games, Inc. All Rights Reserved.
+
+#include "CoreTypes.h"
+
+namespace UE::LowLevelTests
+{
+
+class ITestRunner
+{
+public:
+	/** Returns the test runner if one is active, otherwise null. */
+	static ITestRunner* Get();
+
+	/** Returns true if the test runner will perform global setup and teardown. */
+	virtual bool HasGlobalSetup() const = 0;
+	/** Returns true if the test runner has default logging enabled via UE_LOG. */
+	virtual bool HasLogOutput() const = 0;
+	/** Returns true if the test runner has extra debug functionality enabled. */
+	virtual bool IsDebugMode() const = 0;
+
+protected:
+	ITestRunner();
+	virtual ~ITestRunner();
+};
+
+} // UE::LowLevelTests
diff --git a/Engine/Source/Developer/LowLevelTestsRunner/Public/TestCommon/ApplicationCoreUtilities.h b/Engine/Source/Developer/LowLevelTestsRunner/Public/TestCommon/ApplicationCoreUtilities.h
deleted file mode 100644
diff --git a/Engine/Source/Developer/LowLevelTestsRunner/Public/TestCommon/CoreUtilities.h b/Engine/Source/Developer/LowLevelTestsRunner/Public/TestCommon/CoreUtilities.h
--- a/Engine/Source/Developer/LowLevelTestsRunner/Public/TestCommon/CoreUtilities.h
+++ b/Engine/Source/Developer/LowLevelTestsRunner/Public/TestCommon/CoreUtilities.h
@@ -5,12 +5,6 @@
 #include "CoreMinimal.h"
 #include "Async/TaskGraphInterfaces.h"
 
-void InitCommandLine(bool AllowLogging = false);
-
-void CleanupCommandLine();
-
-void InitOutputDevices();
-
 void InitIOThreadPool(bool MultiThreaded = false, int32 StackSize = 128 * 1024);
 
 void InitThreadPool(bool MultiThreaded = false, int32 StackSize = 128 * 1024);
diff --git a/Engine/Source/Developer/LowLevelTestsRunner/Public/TestGlobals.h b/Engine/Source/Developer/LowLevelTestsRunner/Public/TestGlobals.h
deleted file mode 100644
diff --git a/Engine/Source/Developer/LowLevelTestsRunner/Public/TestListeners/ConsoleListener.h b/Engine/Source/Developer/LowLevelTestsRunner/Public/TestListeners/ConsoleListener.h
deleted file mode 100644
