LowLevelTests: Allow command line arguments to the engine following "--"

The --base-global-module argument is now expected to be in the format --base-global-module=ModuleName.

RunTests has been refactored into the new private FTestRunner type.

diff --git a/Engine/Source/Developer/LowLevelTestsRunner/Private/TestRunner.cpp b/Engine/Source/Developer/LowLevelTestsRunner/Private/TestRunner.cpp
--- a/Engine/Source/Developer/LowLevelTestsRunner/Private/TestRunner.cpp
+++ b/Engine/Source/Developer/LowLevelTestsRunner/Private/TestRunner.cpp
@@ -5,11 +5,13 @@
 
 #include "HAL/PlatformOutputDevices.h"
 #include "HAL/PlatformTLS.h"
+#include "Logging/LogSuppressionInterface.h"
 #include "LowLevelTestModule.h"
+#include "Misc/ScopeExit.h"
+#include "Misc/StringBuilder.h"
 #include "Modules/ModuleManager.h"
-#include "Templates/SharedPointer.h"
+#include "String/Find.h"
 #include "TestCommon/CoreUtilities.h"
-#include "TestHarness.h"
 
 #if WITH_APPLICATION_CORE
 #include "HAL/PlatformApplicationMisc.h"
@@ -18,7 +20,6 @@
 #include <catch2/catch_session.hpp>
 
 #include <iostream>
-#include <set>
 
 namespace UE::LowLevelTests
 {
@@ -45,45 +46,160 @@ ITestRunner::~ITestRunner()
 class FTestRunner final : public ITestRunner
 {
 public:
+	FTestRunner();
+
+	void ParseCommandLine(TConstArrayView<const ANSICHAR*> Args);
+
+	void GlobalSetup() const;
+	void GlobalTeardown() const;
+	void Terminate() const;
+
+	int32 RunCatchSession() const;
+
 	bool HasGlobalSetup() const final { return bGlobalSetup; }
 	bool HasLogOutput() const final { return bLogOutput || bDebugMode; }
 	bool IsDebugMode() const final { return bDebugMode; }
 
-	void SetGlobalSetup(bool bInGlobalSetup) { bGlobalSetup = bInGlobalSetup; }
-	void SetLogOutput(bool bInLogOutput) { bLogOutput = bInLogOutput; }
-	void SetDebugMode(bool bInDebugMode) { bDebugMode = bInDebugMode; }
-
 private:
+	static TArray<FName> GetGlobalModuleNames()
+	{
+		TArray<FName> ModuleNames;
+		FModuleManager::Get().FindModules(TEXT("*GlobalLowLevelTests"), ModuleNames);
+		return ModuleNames;
+	}
+
+	TArray<const ANSICHAR*> CatchArgs;
+	FStringBuilderBase ExtraArgs;
+	FString BaseModuleName;
 	bool bGlobalSetup = true;
 	bool bLogOutput = false;
 	bool bDebugMode = false;
+	bool bMultiThreaded = false;
+	bool bWaitForInputToTerminate = false;
 };
 
-static void LoadBaseTestModule(const FString& BaseModuleName)
+FTestRunner::FTestRunner()
 {
-	if (!BaseModuleName.IsEmpty())
-	{
-		FModuleManager::Get().LoadModule(*BaseModuleName);
-	}
+	// Start setting up the Game Thread.
+	GGameThreadId = FPlatformTLS::GetCurrentThreadId();
+	GIsGameThreadIdInitialized = true;
 }
 
-static void UnloadBaseTestModule(const FString& BaseModuleName)
+void FTestRunner::ParseCommandLine(TConstArrayView<const ANSICHAR*> Args)
 {
-	if (!BaseModuleName.IsEmpty())
+	bool bExtraArg = false;
+	for (FAnsiStringView Arg : Args)
 	{
-		FModuleManager::Get().UnloadModule(*BaseModuleName);
+		if (bExtraArg)
+		{
+			if (const int32 SpaceIndex = String::FindFirstChar(Arg, ' '); SpaceIndex != INDEX_NONE)
+			{
+				if (const int32 EqualIndex = String::FindFirstChar(Arg, '='); EqualIndex != INDEX_NONE && EqualIndex < SpaceIndex)
+				{
+					ExtraArgs.Append(Arg.Left(EqualIndex + 1));
+					Arg.RightChopInline(EqualIndex + 1);
+				}
+				ExtraArgs.AppendChar('"').Append(Arg).AppendChar('"').AppendChar(' ');
+			}
+			else
+			{
+				ExtraArgs.Append(Arg).AppendChar(' ');
+			}
+		}
+		else if (Arg == ANSITEXTVIEW("--"))
+		{
+			bExtraArg = true;
+		}
+		else if (Arg == ANSITEXTVIEW("--global-setup"))
+		{
+			bGlobalSetup = true;
+		}
+		else if (Arg == ANSITEXTVIEW("--no-global-setup"))
+		{
+			bGlobalSetup = false;
+		}
+		else if (Arg == ANSITEXTVIEW("--log"))
+		{
+			bLogOutput = true;
+		}
+		else if (Arg == ANSITEXTVIEW("--no-log"))
+		{
+			bLogOutput = false;
+		}
+		else if (Arg == ANSITEXTVIEW("--debug"))
+		{
+			bDebugMode = true;
+		}
+		else if (Arg == ANSITEXTVIEW("--mt"))
+		{
+			bMultiThreaded = true;
+		}
+		else if (Arg == ANSITEXTVIEW("--no-mt"))
+		{
+			bMultiThreaded = false;
+		}
+		else if (Arg == ANSITEXTVIEW("--wait"))
+		{
+			bWaitForInputToTerminate = true;
+		}
+		else if (Arg == ANSITEXTVIEW("--no-wait"))
+		{
+			bWaitForInputToTerminate = true;
+		}
+		else if (Arg.StartsWith(ANSITEXTVIEW("--base-global-module=")))
+		{
+			BaseModuleName = FString(Arg.RightChop(ANSITEXTVIEW("--base-global-module=").Len()));
+		}
+		else
+		{
+			CatchArgs.Add(Arg.GetData());
+		}
 	}
-}
 
-static TArray<FName> GetGlobalModuleNames()
-{
-	TArray<FName> ModuleNames;
-	FModuleManager::Get().FindModules(TEXT("*GlobalLowLevelTests"), ModuleNames);
-	return ModuleNames;
+	// Break in the debugger on failed assertions when attached.
+	CatchArgs.Add("--break");
 }
 
-static void GlobalModuleSetup()
+void FTestRunner::GlobalSetup() const
 {
+	if (!bGlobalSetup)
+	{
+		return;
+	}
+
+	FCommandLine::Set(*ExtraArgs);
+
+	// Finish setting up the Game Thread, which requires the command line.
+	FPlatformProcess::SetThreadAffinityMask(FPlatformAffinity::GetMainGameMask());
+	FPlatformProcess::SetupGameThread();
+
+	// Always set up GError to handle FatalError, failed assertions, and crashes and other fatal errors.
+#if WITH_APPLICATION_CORE
+	GError = FPlatformApplicationMisc::GetErrorOutputDevice();
+#else
+	GError = FPlatformOutputDevices::GetError();
+#endif
+
+	if (bLogOutput || bDebugMode)
+	{
+		// Set up GWarn to handle Error, Warning, Display; but only when log output is enabled.
+	#if WITH_APPLICATION_CORE
+		GWarn = FPlatformApplicationMisc::GetFeedbackContext();
+	#else
+		GWarn = FPlatformOutputDevices::GetFeedbackContext();
+	#endif
+
+		// Set up default output devices to handle Log, Verbose, VeryVerbose.
+		FPlatformOutputDevices::SetupOutputDevices();
+
+		FLogSuppressionInterface::Get().ProcessConfigAndCommandLine();
+	}
+
+	if (!BaseModuleName.IsEmpty())
+	{
+		FModuleManager::Get().LoadModule(*BaseModuleName);
+	}
+
 	for (FName ModuleName : GetGlobalModuleNames())
 	{
 		if (ILowLevelTestsModule* Module = FModuleManager::LoadModulePtr<ILowLevelTestsModule>(ModuleName))
@@ -93,8 +209,13 @@ static void GlobalModuleSetup()
 	}
 }
 
-static void GlobalModuleTeardown()
+void FTestRunner::GlobalTeardown() const
 {
+	if (!bGlobalSetup)
+	{
+		return;
+	}
+
 	for (FName ModuleName : GetGlobalModuleNames())
 	{
 		if (ILowLevelTestsModule* Module = FModuleManager::GetModulePtr<ILowLevelTestsModule>(ModuleName))
@@ -106,172 +227,65 @@ static void GlobalModuleTeardown()
 			}
 		}
 	}
-}
-
-} // UE::LowLevelTests
-
-int RunTests(int argc, const char* argv[])
-{
-	using namespace UE::LowLevelTests;
 
-	// Set up the Game Thread.
-	GGameThreadId = FPlatformTLS::GetCurrentThreadId();
-	GIsGameThreadIdInitialized = true;
-
-#ifdef SLEEP_ON_INIT
-	// Sleep to allow sync with Gauntlet.
-	FPlatformProcess::Sleep(5.0f);
-#endif
-
-	//Read command-line from file (if any). Some platforms do this earlier.
-#ifndef PLATFORM_SKIP_ADDITIONAL_ARGS
-	int ArgsOverrideNum = 0;
-	const char** ArgsOverride = ReadAndAppendAdditionalArgs(GetProcessExecutablePath(), &ArgsOverrideNum, argv, argc);
-	if (ArgsOverride != nullptr && ArgsOverrideNum > 1)
+	if (!BaseModuleName.IsEmpty())
 	{
-		argc = ArgsOverrideNum;
-		argv = ArgsOverride;
+		FModuleManager::Get().UnloadModule(*BaseModuleName);
 	}
-#endif
 
-	int CatchArgc = 0;
-	TUniquePtr<const char* []> CatchArgv = MakeUnique<const char* []>(argc + 1);
-
-	std::set<std::string> KnownLLTArgs;
-	KnownLLTArgs.insert("--wait");
-	KnownLLTArgs.insert("--no-wait");
-	KnownLLTArgs.insert("--log");
-	KnownLLTArgs.insert("--no-log");
-	KnownLLTArgs.insert("--mt");
-	KnownLLTArgs.insert("--no-mt");
-	KnownLLTArgs.insert("--global-setup");
-	KnownLLTArgs.insert("--no-global-setup");
-	KnownLLTArgs.insert("--base-global-module");
-	KnownLLTArgs.insert("--debug");
-
-	// By default don't wait for input
-	bool bWaitForInputToTerminate = false;
+	CleanupPlatform();
 
-	FTestRunner TestRunner;
-	bool bMultiThreaded = false;
-
-	FString BaseModuleName;
+	FCommandLine::Reset();
+}
 
-	// Every argument that is not in the list of known low level test options and is considered to be a catch test runner argument.
-	for (int i = 0; i < argc; ++i)
+void FTestRunner::Terminate() const
+{
+#if PLATFORM_DESKTOP
+	if (bWaitForInputToTerminate)
 	{
-		if (KnownLLTArgs.find(argv[i]) != KnownLLTArgs.end())
-		{
-			if (std::strcmp(argv[i], "--wait") == 0)
-			{
-				bWaitForInputToTerminate = true;
-			}
-			else if (std::strcmp(argv[i], "--no-wait") == 0)
-			{
-				bWaitForInputToTerminate = false;
-			}
-			if (std::strcmp(argv[i], "--log") == 0)
-			{
-				TestRunner.SetLogOutput(true);
-			}
-			else if (std::strcmp(argv[i], "--no-log") == 0)
-			{
-				TestRunner.SetLogOutput(false);
-			}
-			if (std::strcmp(argv[i], "--mt") == 0)
-			{
-				bMultiThreaded = true;
-			}
-			else if (std::strcmp(argv[i], "--no-mt") == 0)
-			{
-				bMultiThreaded = false;
-			}
-			if (std::strcmp(argv[i], "--global-setup") == 0)
-			{
-				TestRunner.SetGlobalSetup(true);
-			}
-			else if (std::strcmp(argv[i], "--no-global-setup") == 0)
-			{
-				TestRunner.SetGlobalSetup(false);
-			}
-			if (std::strcmp(argv[i], "--debug") == 0)
-			{
-				TestRunner.SetDebugMode(true);
-			}
-			// If we have --base-global-module parse proceeding argument as its option-value
-			if (std::strcmp(argv[i], "--base-global-module") == 0 && i + 1 < argc)
-			{
-				BaseModuleName = argv[i + 1];
-				++i;
-			}
-		}
-		else
-		{
-			// Passing to catch2
-			CatchArgv.Get()[CatchArgc++] = argv[i];
-		}
+		std::cout << "Press enter to exit..." << std::endl;
+		std::cin.ignore();
 	}
+#endif
+}
 
-	// Break in the debugger on failed assertions when attached.
-	CatchArgv.Get()[CatchArgc++] = "--break";
+int32 FTestRunner::RunCatchSession() const
+{
+	return Catch::Session().run(CatchArgs.Num(), CatchArgs.GetData());
+}
 
-	if (TestRunner.HasGlobalSetup())
-	{
-		FCommandLine::Set(TEXT(""));
+} // UE::LowLevelTests
 
-		// Finish setting up the Game Thread, which requires the command line.
-		FPlatformProcess::SetThreadAffinityMask(FPlatformAffinity::GetMainGameMask());
-		FPlatformProcess::SetupGameThread();
+int RunTests(int32 ArgC, const ANSICHAR* ArgV[])
+{
+	UE::LowLevelTests::FTestRunner TestRunner;
 
-		// Always set up GError to handle FatalError, failed assertions, and crashes and other fatal errors.
-	#if WITH_APPLICATION_CORE
-		GError = FPlatformApplicationMisc::GetErrorOutputDevice();
-	#else
-		GError = FPlatformOutputDevices::GetError();
-	#endif
+#ifdef SLEEP_ON_INIT
+	// Sleep to allow sync with Gauntlet.
+	FPlatformProcess::Sleep(5.0f);
+#endif
 
-		if (TestRunner.HasLogOutput())
+	// Read command-line from file (if any). Some platforms do this earlier.
+#ifndef PLATFORM_SKIP_ADDITIONAL_ARGS
+	{
+		int32 OverrideArgC = 0;
+		const ANSICHAR** OverrideArgV = ReadAndAppendAdditionalArgs(GetProcessExecutablePath(), &OverrideArgC, ArgV, ArgC);
+		if (OverrideArgV && OverrideArgC > 1)
 		{
-			// Set up GWarn to handle Error, Warning, Display; but only when log output is enabled.
-		#if WITH_APPLICATION_CORE
-			GWarn = FPlatformApplicationMisc::GetFeedbackContext();
-		#else
-			GWarn = FPlatformOutputDevices::GetFeedbackContext();
-		#endif
-
-			// Set up default output devices to handle Log, Verbose, VeryVerbose.
-			FPlatformOutputDevices::SetupOutputDevices();
+			ArgC = OverrideArgC;
+			ArgV = OverrideArgV;
 		}
-
-		LoadBaseTestModule(BaseModuleName);
-
-		GlobalModuleSetup();
 	}
+#endif
 
-	int SessionResult = 0;
-	{
-		SessionResult = Catch::Session().run(CatchArgc, CatchArgv.Get());
-		CatchArgv.Reset();
-	}
-
-	if (TestRunner.HasGlobalSetup())
-	{
-		GlobalModuleTeardown();
-
-		UnloadBaseTestModule(BaseModuleName);
-
-		CleanupPlatform();
+	TestRunner.ParseCommandLine(MakeArrayView(ArgV, ArgC));
+	TestRunner.GlobalSetup();
 
-		FCommandLine::Reset();
-	}
-
-#if PLATFORM_DESKTOP
-	if (bWaitForInputToTerminate)
+	ON_SCOPE_EXIT
 	{
-		std::cout << "Press enter to exit..." << std::endl;
-		std::cin.ignore();
-	}
-#endif
+		TestRunner.GlobalTeardown();
+		TestRunner.Terminate();
+	};
 
-	return SessionResult;
+	return TestRunner.RunCatchSession();
 }
