Test runner:
- sleep on init to allow for sync w/ Gauntlet for Xbox platforms
- remove unused --base-global-module #fyi Stephen.Ma
Gauntlet:
- support low level tests build for different configurations (debug etc)
- readable command line arguments (-r to --reporter)
- allow argument names containing "-" e.g. --filenames-as-tags

diff --git a/Engine/Source/Developer/LowLevelTestsRunner/Private/TestRunner.cpp b/Engine/Source/Developer/LowLevelTestsRunner/Private/TestRunner.cpp
--- a/Engine/Source/Developer/LowLevelTestsRunner/Private/TestRunner.cpp
+++ b/Engine/Source/Developer/LowLevelTestsRunner/Private/TestRunner.cpp
@@ -24,235 +24,236 @@
 namespace UE::LowLevelTests
 {
 
-static ITestRunner* GTestRunner;
+	static ITestRunner* GTestRunner;
 
-ITestRunner* ITestRunner::Get()
-{
-	return GTestRunner;
-}
+	ITestRunner* ITestRunner::Get()
+	{
+		return GTestRunner;
+	}
 
-ITestRunner::ITestRunner()
-{
-	check(!GTestRunner);
-	GTestRunner = this;
-}
+	ITestRunner::ITestRunner()
+	{
+		check(!GTestRunner);
+		GTestRunner = this;
+	}
 
-ITestRunner::~ITestRunner()
-{
-	check(GTestRunner == this);
-	GTestRunner = nullptr;
-}
+	ITestRunner::~ITestRunner()
+	{
+		check(GTestRunner == this);
+		GTestRunner = nullptr;
+	}
 
-class FTestRunner final : public ITestRunner
-{
-public:
-	FTestRunner();
+	class FTestRunner final : public ITestRunner
+	{
+	public:
+		FTestRunner();
+
+		void ParseCommandLine(TConstArrayView<const ANSICHAR*> Args);
 
-	void ParseCommandLine(TConstArrayView<const ANSICHAR*> Args);
+		void SleepOnInit() const;
 
-	void GlobalSetup() const;
-	void GlobalTeardown() const;
-	void Terminate() const;
+		void GlobalSetup() const;
+		void GlobalTeardown() const;
+		void Terminate() const;
 
-	int32 RunCatchSession() const;
+		int32 RunCatchSession() const;
 
-	bool HasGlobalSetup() const final { return bGlobalSetup; }
-	bool HasLogOutput() const final { return bLogOutput || bDebugMode; }
-	bool IsDebugMode() const final { return bDebugMode; }
+		bool HasGlobalSetup() const final { return bGlobalSetup; }
+		bool HasLogOutput() const final { return bLogOutput || bDebugMode; }
+		bool IsDebugMode() const final { return bDebugMode; }
 
-private:
-	static TArray<FName> GetGlobalModuleNames()
+	private:
+		static TArray<FName> GetGlobalModuleNames()
+		{
+			TArray<FName> ModuleNames;
+			FModuleManager::Get().FindModules(TEXT("*GlobalLowLevelTests"), ModuleNames);
+			return ModuleNames;
+		}
+
+		TArray<const ANSICHAR*> CatchArgs;
+		FStringBuilderBase ExtraArgs;
+		bool bGlobalSetup = true;
+		bool bLogOutput = false;
+		bool bDebugMode = false;
+		bool bMultiThreaded = false;
+		bool bWaitForInputToTerminate = false;
+		int32 SleepOnInitMilliseconds = 0;
+	};
+
+	FTestRunner::FTestRunner()
 	{
-		TArray<FName> ModuleNames;
-		FModuleManager::Get().FindModules(TEXT("*GlobalLowLevelTests"), ModuleNames);
-		return ModuleNames;
+		// Start setting up the Game Thread.
+		GGameThreadId = FPlatformTLS::GetCurrentThreadId();
+		GIsGameThreadIdInitialized = true;
 	}
 
-	TArray<const ANSICHAR*> CatchArgs;
-	FStringBuilderBase ExtraArgs;
-	FString BaseModuleName;
-	bool bGlobalSetup = true;
-	bool bLogOutput = false;
-	bool bDebugMode = false;
-	bool bMultiThreaded = false;
-	bool bWaitForInputToTerminate = false;
-};
-
-FTestRunner::FTestRunner()
-{
-	// Start setting up the Game Thread.
-	GGameThreadId = FPlatformTLS::GetCurrentThreadId();
-	GIsGameThreadIdInitialized = true;
-}
-
-void FTestRunner::ParseCommandLine(TConstArrayView<const ANSICHAR*> Args)
-{
-	bool bExtraArg = false;
-	for (FAnsiStringView Arg : Args)
+	void FTestRunner::ParseCommandLine(TConstArrayView<const ANSICHAR*> Args)
 	{
-		if (bExtraArg)
+		bool bExtraArg = false;
+		for (FAnsiStringView Arg : Args)
 		{
-			if (const int32 SpaceIndex = String::FindFirstChar(Arg, ' '); SpaceIndex != INDEX_NONE)
+			if (bExtraArg)
 			{
-				if (const int32 EqualIndex = String::FindFirstChar(Arg, '='); EqualIndex != INDEX_NONE && EqualIndex < SpaceIndex)
+				if (const int32 SpaceIndex = String::FindFirstChar(Arg, ' '); SpaceIndex != INDEX_NONE)
+				{
+					if (const int32 EqualIndex = String::FindFirstChar(Arg, '='); EqualIndex != INDEX_NONE && EqualIndex < SpaceIndex)
+					{
+						ExtraArgs.Append(Arg.Left(EqualIndex + 1));
+						Arg.RightChopInline(EqualIndex + 1);
+					}
+					ExtraArgs.AppendChar('"').Append(Arg).AppendChar('"').AppendChar(' ');
+				}
+				else
 				{
-					ExtraArgs.Append(Arg.Left(EqualIndex + 1));
-					Arg.RightChopInline(EqualIndex + 1);
+					ExtraArgs.Append(Arg).AppendChar(' ');
 				}
-				ExtraArgs.AppendChar('"').Append(Arg).AppendChar('"').AppendChar(' ');
+			}
+			else if (Arg == ANSITEXTVIEW("--"))
+			{
+				bExtraArg = true;
+			}
+			else if (Arg.StartsWith(ANSITEXTVIEW("--sleep=")))
+			{
+				SleepOnInitMilliseconds = std::atoi(Arg.LeftChop(7).GetData());
+			}
+			else if (Arg == ANSITEXTVIEW("--global-setup"))
+			{
+				bGlobalSetup = true;
+			}
+			else if (Arg == ANSITEXTVIEW("--no-global-setup"))
+			{
+				bGlobalSetup = false;
+			}
+			else if (Arg == ANSITEXTVIEW("--log"))
+			{
+				bLogOutput = true;
+			}
+			else if (Arg == ANSITEXTVIEW("--no-log"))
+			{
+				bLogOutput = false;
+			}
+			else if (Arg == ANSITEXTVIEW("--debug"))
+			{
+				bDebugMode = true;
+			}
+			else if (Arg == ANSITEXTVIEW("--mt"))
+			{
+				bMultiThreaded = true;
+			}
+			else if (Arg == ANSITEXTVIEW("--no-mt"))
+			{
+				bMultiThreaded = false;
+			}
+			else if (Arg == ANSITEXTVIEW("--wait"))
+			{
+				bWaitForInputToTerminate = true;
+			}
+			else if (Arg == ANSITEXTVIEW("--no-wait"))
+			{
+				bWaitForInputToTerminate = true;
 			}
 			else
 			{
-				ExtraArgs.Append(Arg).AppendChar(' ');
+				CatchArgs.Add(Arg.GetData());
 			}
 		}
-		else if (Arg == ANSITEXTVIEW("--"))
-		{
-			bExtraArg = true;
-		}
-		else if (Arg == ANSITEXTVIEW("--global-setup"))
-		{
-			bGlobalSetup = true;
-		}
-		else if (Arg == ANSITEXTVIEW("--no-global-setup"))
-		{
-			bGlobalSetup = false;
-		}
-		else if (Arg == ANSITEXTVIEW("--log"))
-		{
-			bLogOutput = true;
-		}
-		else if (Arg == ANSITEXTVIEW("--no-log"))
-		{
-			bLogOutput = false;
-		}
-		else if (Arg == ANSITEXTVIEW("--debug"))
-		{
-			bDebugMode = true;
-		}
-		else if (Arg == ANSITEXTVIEW("--mt"))
-		{
-			bMultiThreaded = true;
-		}
-		else if (Arg == ANSITEXTVIEW("--no-mt"))
-		{
-			bMultiThreaded = false;
-		}
-		else if (Arg == ANSITEXTVIEW("--wait"))
-		{
-			bWaitForInputToTerminate = true;
-		}
-		else if (Arg == ANSITEXTVIEW("--no-wait"))
-		{
-			bWaitForInputToTerminate = true;
-		}
-		else if (Arg.StartsWith(ANSITEXTVIEW("--base-global-module=")))
-		{
-			BaseModuleName = FString(Arg.RightChop(ANSITEXTVIEW("--base-global-module=").Len()));
-		}
-		else
+
+		// Break in the debugger on failed assertions when attached.
+		CatchArgs.Add("--break");
+	}
+
+	void FTestRunner::SleepOnInit() const
+	{
+		if (SleepOnInitMilliseconds)
 		{
-			CatchArgs.Add(Arg.GetData());
+			// Sleep to allow sync with Gauntlet.
+			FPlatformProcess::Sleep(SleepOnInitMilliseconds * 1000);
 		}
 	}
 
-	// Break in the debugger on failed assertions when attached.
-	CatchArgs.Add("--break");
-}
-
-void FTestRunner::GlobalSetup() const
-{
-	if (!bGlobalSetup)
+	void FTestRunner::GlobalSetup() const
 	{
-		return;
-	}
+		if (!bGlobalSetup)
+		{
+			return;
+		}
 
-	FCommandLine::Set(*ExtraArgs);
+		FCommandLine::Set(*ExtraArgs);
 
-	// Finish setting up the Game Thread, which requires the command line.
-	FPlatformProcess::SetThreadAffinityMask(FPlatformAffinity::GetMainGameMask());
-	FPlatformProcess::SetupGameThread();
+		// Finish setting up the Game Thread, which requires the command line.
+		FPlatformProcess::SetThreadAffinityMask(FPlatformAffinity::GetMainGameMask());
+		FPlatformProcess::SetupGameThread();
 
-	// Always set up GError to handle FatalError, failed assertions, and crashes and other fatal errors.
+		// Always set up GError to handle FatalError, failed assertions, and crashes and other fatal errors.
 #if WITH_APPLICATION_CORE
-	GError = FPlatformApplicationMisc::GetErrorOutputDevice();
+		GError = FPlatformApplicationMisc::GetErrorOutputDevice();
 #else
-	GError = FPlatformOutputDevices::GetError();
+		GError = FPlatformOutputDevices::GetError();
 #endif
 
-	if (bLogOutput || bDebugMode)
-	{
-		// Set up GWarn to handle Error, Warning, Display; but only when log output is enabled.
-	#if WITH_APPLICATION_CORE
-		GWarn = FPlatformApplicationMisc::GetFeedbackContext();
-	#else
-		GWarn = FPlatformOutputDevices::GetFeedbackContext();
-	#endif
-
-		// Set up default output devices to handle Log, Verbose, VeryVerbose.
-		FPlatformOutputDevices::SetupOutputDevices();
+		if (bLogOutput || bDebugMode)
+		{
+			// Set up GWarn to handle Error, Warning, Display; but only when log output is enabled.
+#if WITH_APPLICATION_CORE
+			GWarn = FPlatformApplicationMisc::GetFeedbackContext();
+#else
+			GWarn = FPlatformOutputDevices::GetFeedbackContext();
+#endif
 
-		FLogSuppressionInterface::Get().ProcessConfigAndCommandLine();
-	}
+			// Set up default output devices to handle Log, Verbose, VeryVerbose.
+			FPlatformOutputDevices::SetupOutputDevices();
 
-	if (!BaseModuleName.IsEmpty())
-	{
-		FModuleManager::Get().LoadModule(*BaseModuleName);
-	}
+			FLogSuppressionInterface::Get().ProcessConfigAndCommandLine();
+		}
 
-	for (FName ModuleName : GetGlobalModuleNames())
-	{
-		if (ILowLevelTestsModule* Module = FModuleManager::LoadModulePtr<ILowLevelTestsModule>(ModuleName))
+		for (FName ModuleName : GetGlobalModuleNames())
 		{
-			Module->GlobalSetup();
+			if (ILowLevelTestsModule* Module = FModuleManager::LoadModulePtr<ILowLevelTestsModule>(ModuleName))
+			{
+				Module->GlobalSetup();
+			}
 		}
 	}
-}
 
-void FTestRunner::GlobalTeardown() const
-{
-	if (!bGlobalSetup)
+	void FTestRunner::GlobalTeardown() const
 	{
-		return;
-	}
+		if (!bGlobalSetup)
+		{
+			return;
+		}
 
-	for (FName ModuleName : GetGlobalModuleNames())
-	{
-		if (ILowLevelTestsModule* Module = FModuleManager::GetModulePtr<ILowLevelTestsModule>(ModuleName))
+		for (FName ModuleName : GetGlobalModuleNames())
 		{
-			Module->GlobalTeardown();
-			if (Module->SupportsAutomaticShutdown())
+			if (ILowLevelTestsModule* Module = FModuleManager::GetModulePtr<ILowLevelTestsModule>(ModuleName))
 			{
-				Module->ShutdownModule();
+				Module->GlobalTeardown();
+				if (Module->SupportsAutomaticShutdown())
+				{
+					Module->ShutdownModule();
+				}
 			}
 		}
+
+		CleanupPlatform();
+
+		FCommandLine::Reset();
 	}
 
-	if (!BaseModuleName.IsEmpty())
+	void FTestRunner::Terminate() const
 	{
-		FModuleManager::Get().UnloadModule(*BaseModuleName);
+#if PLATFORM_DESKTOP
+		if (bWaitForInputToTerminate)
+		{
+			std::cout << "Press enter to exit..." << std::endl;
+			std::cin.ignore();
+		}
+#endif
 	}
 
-	CleanupPlatform();
-
-	FCommandLine::Reset();
-}
-
-void FTestRunner::Terminate() const
-{
-#if PLATFORM_DESKTOP
-	if (bWaitForInputToTerminate)
+	int32 FTestRunner::RunCatchSession() const
 	{
-		std::cout << "Press enter to exit..." << std::endl;
-		std::cin.ignore();
+		return Catch::Session().run(CatchArgs.Num(), CatchArgs.GetData());
 	}
-#endif
-}
-
-int32 FTestRunner::RunCatchSession() const
-{
-	return Catch::Session().run(CatchArgs.Num(), CatchArgs.GetData());
-}
 
 } // UE::LowLevelTests
 
@@ -260,11 +261,6 @@ int RunTests(int32 ArgC, const ANSICHAR* ArgV[])
 {
 	UE::LowLevelTests::FTestRunner TestRunner;
 
-#ifdef SLEEP_ON_INIT
-	// Sleep to allow sync with Gauntlet.
-	FPlatformProcess::Sleep(5.0f);
-#endif
-
 	// Read command-line from file (if any). Some platforms do this earlier.
 #ifndef PLATFORM_SKIP_ADDITIONAL_ARGS
 	{
@@ -279,6 +275,9 @@ int RunTests(int32 ArgC, const ANSICHAR* ArgV[])
 #endif
 
 	TestRunner.ParseCommandLine(MakeArrayView(ArgV, ArgC));
+
+	TestRunner.SleepOnInit();
+	
 	TestRunner.GlobalSetup();
 
 	ON_SCOPE_EXIT
diff --git a/Engine/Source/Programs/AutomationTool/Gauntlet/Unreal/Base/Gauntlet.UnrealTestConfiguration.cs b/Engine/Source/Programs/AutomationTool/Gauntlet/Unreal/Base/Gauntlet.UnrealTestConfiguration.cs
--- a/Engine/Source/Programs/AutomationTool/Gauntlet/Unreal/Base/Gauntlet.UnrealTestConfiguration.cs
+++ b/Engine/Source/Programs/AutomationTool/Gauntlet/Unreal/Base/Gauntlet.UnrealTestConfiguration.cs
@@ -53,7 +53,7 @@ namespace Gauntlet
 		public void AddRawCommandline(string InRawCommandline, bool bOverrideExistingValues = true)
 		{
 			// turn Name(p1,etc) into a collection of Name|(p1,etc) groups
-			MatchCollection Matches = Regex.Matches(InRawCommandline, "-(?<option>\\-?[\\w\\d.:!\\[\\]\\/\\\\]+)(=(?<value>(\"([^\"]*)\")|(\\S+)))?");
+			MatchCollection Matches = Regex.Matches(InRawCommandline, "-(?<option>\\-?[\\w\\d.:!\\[\\]\\/\\\\\\-]+)(=(?<value>(\"([^\"]*)\")|(\\S+)))?");
 
 			foreach (Match M in Matches)
 			{
diff --git a/Engine/Source/Programs/AutomationTool/LowLevelTests/RunLowLevelTests.cs b/Engine/Source/Programs/AutomationTool/LowLevelTests/RunLowLevelTests.cs
--- a/Engine/Source/Programs/AutomationTool/LowLevelTests/RunLowLevelTests.cs
+++ b/Engine/Source/Programs/AutomationTool/LowLevelTests/RunLowLevelTests.cs
@@ -157,6 +157,9 @@ namespace LowLevelTests
 
 		public string Tags;
 
+		[AutoParam(0)]
+		public int Sleep;
+
 		public string Build;
 
 		[AutoParam("")]
@@ -227,13 +230,15 @@ namespace LowLevelTests
 		public IAppInstance Instance { get; protected set; }
 		private LowLevelTestsBuildSource BuildSource { get; set; }
 		private string Tags { get; set; }
+		private int Sleep { get; set; }
 
 		public UnrealDeviceReservation UnrealDeviceReservation { get; private set; }
 
-		public LowLevelTestsSession(LowLevelTestsBuildSource InBuildSource, string InTags)
+		public LowLevelTestsSession(LowLevelTestsBuildSource InBuildSource, string InTags, int InSleep)
 		{
 			BuildSource = InBuildSource;
 			Tags = InTags;
+			Sleep = InSleep;
 			UnrealDeviceReservation = new UnrealDeviceReservation();
 		}
 
@@ -257,7 +262,7 @@ namespace LowLevelTests
 
 			// TargetDevice<Platform> classes have a hard dependency on UnrealAppConfig instead of IAppConfig.
 			// More refactoring needed to support non-packaged applications that can be run natively from a path on the device.
-			UnrealAppConfig AppConfig = BuildSource.GetUnrealAppConfig(Tags);
+			UnrealAppConfig AppConfig = BuildSource.GetUnrealAppConfig(Tags, Sleep);
 
 			IEnumerable<ITargetDevice> DevicesToInstallOn = UnrealDeviceReservation.ReservedDevices.ToArray();
 			ITargetDevice Device = DevicesToInstallOn.Where(D => D.IsConnected && D.Platform == BuildSource.Platform).First();
@@ -274,7 +279,7 @@ namespace LowLevelTests
 			{
 				InstallSuccess = false;
 
-				Log.Info("Failed to install low level tests app onto device {0}: {1}", Device, Ex);
+				Log.Info("Failed to install low level tests app onto device {0}: {1}", Device, Ex.ToString());
 
 				UnrealDeviceReservation.MarkProblemDevice(Device);
 				IDeviceUsageReporter.RecordEnd(Device.Name, (UnrealTargetPlatform)Device.Platform, IDeviceUsageReporter.EventType.Install, IDeviceUsageReporter.EventState.Failure);
@@ -555,7 +560,7 @@ namespace LowLevelTests
 				.First();
 		}
 
-		public UnrealAppConfig GetUnrealAppConfig(string InTags)
+		public UnrealAppConfig GetUnrealAppConfig(string InTags, int InSleep)
 		{
 			if (CachedConfig == null)
 			{
@@ -571,16 +576,20 @@ namespace LowLevelTests
 				// Set reporting options, filters etc
 				if (TestApp == "LowLevelTests" || TestApp == "LowLevelTestsTests")
 				{
-					CachedConfig.CommandLineParams.Add("d", "yes");
+					CachedConfig.CommandLineParams.AddRawCommandline("--durations=yes");
 				}
-				CachedConfig.CommandLineParams.Add("r", "console");
-				CachedConfig.CommandLineParams.Add("o", LowLevelTestsReporting.GetTargetReportPath(Platform, TestApp, BuildPath));
-				CachedConfig.CommandLineParams.Add("-#", null, true);
+				CachedConfig.CommandLineParams.AddRawCommandline("--reporter=console");
+				CachedConfig.CommandLineParams.AddRawCommandline(string.Format("--out={0}", LowLevelTestsReporting.GetTargetReportPath(Platform, TestApp, BuildPath)));
+				CachedConfig.CommandLineParams.AddRawCommandline("--filenames-as-tags");
 				if (!string.IsNullOrEmpty(InTags))
 				{
 					CachedConfig.CommandLineParams.Add(InTags, null, true);
 				}
-				CachedConfig.CommandLineParams.Add("--debug", null, true);
+				if (InSleep > 0)
+				{
+					CachedConfig.CommandLineParams.AddRawCommandline(String.Format("--sleep={0}", InSleep));
+				}
+				CachedConfig.CommandLineParams.AddRawCommandline("--debug");
 			}
 			return CachedConfig;
 		}
@@ -601,4 +610,4 @@ namespace LowLevelTests
 			Flags = BuildFlags.CanReplaceExecutable | BuildFlags.Loose;
 		}
 	}
-}
\ No newline at end of file
+}
diff --git a/Engine/Source/Programs/AutomationTool/LowLevelTests/Tests/Gauntlet.LowLevelTestNode.cs b/Engine/Source/Programs/AutomationTool/LowLevelTests/Tests/Gauntlet.LowLevelTestNode.cs
--- a/Engine/Source/Programs/AutomationTool/LowLevelTests/Tests/Gauntlet.LowLevelTestNode.cs
+++ b/Engine/Source/Programs/AutomationTool/LowLevelTests/Tests/Gauntlet.LowLevelTestNode.cs
@@ -41,7 +41,7 @@ namespace LowLevelTests
 		{
 			if (LowLevelTestsApp == null)
 			{
-				LowLevelTestsApp = new LowLevelTestsSession(Context.BuildInfo, Context.Options.Tags);
+				LowLevelTestsApp = new LowLevelTestsSession(Context.BuildInfo, Context.Options.Tags, Context.Options.Sleep);
 			}
 
 			return LowLevelTestsApp.TryReserveDevices();
