RightChop instead of LeftChop
#rnx

diff --git a/Engine/Source/Developer/LowLevelTestsRunner/Private/TestRunner.cpp b/Engine/Source/Developer/LowLevelTestsRunner/Private/TestRunner.cpp
--- a/Engine/Source/Developer/LowLevelTestsRunner/Private/TestRunner.cpp
+++ b/Engine/Source/Developer/LowLevelTestsRunner/Private/TestRunner.cpp
@@ -114,7 +114,7 @@ namespace UE::LowLevelTests
 			}
 			else if (Arg.StartsWith(ANSITEXTVIEW("--sleep=")))
 			{
-				SleepOnInitMilliseconds = std::atoi(Arg.LeftChop(7).GetData());
+				SleepOnInitMilliseconds = std::atoi(Arg.RightChop(8).GetData());
 			}
 			else if (Arg == ANSITEXTVIEW("--global-setup"))
 			{
@@ -167,7 +167,7 @@ namespace UE::LowLevelTests
 		if (SleepOnInitMilliseconds)
 		{
 			// Sleep to allow sync with Gauntlet.
-			FPlatformProcess::Sleep(SleepOnInitMilliseconds * 1000);
+			FPlatformProcess::Sleep(SleepOnInitMilliseconds);
 		}
 	}
 
