Sleep uses seconds, not milliseconds
#rnx

diff --git a/Engine/Source/Developer/LowLevelTestsRunner/Private/TestRunner.cpp b/Engine/Source/Developer/LowLevelTestsRunner/Private/TestRunner.cpp
--- a/Engine/Source/Developer/LowLevelTestsRunner/Private/TestRunner.cpp
+++ b/Engine/Source/Developer/LowLevelTestsRunner/Private/TestRunner.cpp
@@ -77,7 +77,7 @@ namespace UE::LowLevelTests
 		bool bDebugMode = false;
 		bool bMultiThreaded = false;
 		bool bWaitForInputToTerminate = false;
-		int32 SleepOnInitMilliseconds = 0;
+		int32 SleepOnInitSeconds = 0;
 	};
 
 	FTestRunner::FTestRunner()
@@ -114,7 +114,7 @@ namespace UE::LowLevelTests
 			}
 			else if (Arg.StartsWith(ANSITEXTVIEW("--sleep=")))
 			{
-				SleepOnInitMilliseconds = std::atoi(Arg.RightChop(8).GetData());
+				SleepOnInitSeconds = std::atoi(Arg.RightChop(8).GetData());
 			}
 			else if (Arg == ANSITEXTVIEW("--global-setup"))
 			{
@@ -164,10 +164,10 @@ namespace UE::LowLevelTests
 
 	void FTestRunner::SleepOnInit() const
 	{
-		if (SleepOnInitMilliseconds)
+		if (SleepOnInitSeconds)
 		{
 			// Sleep to allow sync with Gauntlet.
-			FPlatformProcess::Sleep(SleepOnInitMilliseconds);
+			FPlatformProcess::Sleep(SleepOnInitSeconds);
 		}
 	}
 
