LLT TestRunner non-STL argument value parsing

diff --git a/Engine/Source/Developer/LowLevelTestsRunner/Private/TestRunner.cpp b/Engine/Source/Developer/LowLevelTestsRunner/Private/TestRunner.cpp
--- a/Engine/Source/Developer/LowLevelTestsRunner/Private/TestRunner.cpp
+++ b/Engine/Source/Developer/LowLevelTestsRunner/Private/TestRunner.cpp
@@ -11,6 +11,7 @@
 #include "Misc/StringBuilder.h"
 #include "Modules/ModuleManager.h"
 #include "String/Find.h"
+#include "String/LexFromString.h"
 #include "TestCommon/CoreUtilities.h"
 
 #if WITH_APPLICATION_CORE
@@ -114,7 +115,7 @@ namespace UE::LowLevelTests
 			}
 			else if (Arg.StartsWith(ANSITEXTVIEW("--sleep=")))
 			{
-				SleepOnInitSeconds = std::atoi(Arg.RightChop(8).GetData());
+				LexFromString(SleepOnInitSeconds, WriteToString<16>(Arg.RightChop(8)).ToView());
 			}
 			else if (Arg == ANSITEXTVIEW("--global-setup"))
 			{
