LowLevelTests: Removed an unnecessary #include

#rnx

diff --git a/Engine/Source/Developer/LowLevelTestsRunner/Public/LowLevelTestModule.h b/Engine/Source/Developer/LowLevelTestsRunner/Public/LowLevelTestModule.h
--- a/Engine/Source/Developer/LowLevelTestsRunner/Public/LowLevelTestModule.h
+++ b/Engine/Source/Developer/LowLevelTestsRunner/Public/LowLevelTestModule.h
@@ -3,7 +3,6 @@
 #pragma once
 
 #include "Modules/ModuleManager.h"
-#include "Templates/SharedPointer.h"
 
 class ILowLevelTestsModule : public IModuleInterface
 {
