LowLevelTests: Always shut down modules on exit

It is unsafe to skip module shutdown because many modules need to destroy things that were lazily initialized, and may crash later during shutdown if the expected shutdown flow was skipped.

#rnx

diff --git a/Engine/Source/Developer/LowLevelTestsRunner/Private/TestCommon/CoreUtilities.cpp b/Engine/Source/Developer/LowLevelTestsRunner/Private/TestCommon/CoreUtilities.cpp
--- a/Engine/Source/Developer/LowLevelTestsRunner/Private/TestCommon/CoreUtilities.cpp
+++ b/Engine/Source/Developer/LowLevelTestsRunner/Private/TestCommon/CoreUtilities.cpp
@@ -104,5 +104,4 @@ void CleanupTaskGraphAndDependencies()
 void CleanupPlatform()
 {
 	FPlatformMisc::PlatformTearDown();
-	FGenericPlatformMisc::RequestExit(false);
 }
\ No newline at end of file
diff --git a/Engine/Source/Developer/LowLevelTestsRunner/Private/TestRunner.cpp b/Engine/Source/Developer/LowLevelTestsRunner/Private/TestRunner.cpp
--- a/Engine/Source/Developer/LowLevelTestsRunner/Private/TestRunner.cpp
+++ b/Engine/Source/Developer/LowLevelTestsRunner/Private/TestRunner.cpp
@@ -228,10 +228,6 @@ namespace UE::LowLevelTests
 			if (ILowLevelTestsModule* Module = FModuleManager::GetModulePtr<ILowLevelTestsModule>(ModuleName))
 			{
 				Module->GlobalTeardown();
-				if (Module->SupportsAutomaticShutdown())
-				{
-					Module->ShutdownModule();
-				}
 			}
 		}
 
@@ -285,6 +281,8 @@ int RunTests(int32 ArgC, const ANSICHAR* ArgV[])
 	{
 		TestRunner.GlobalTeardown();
 		TestRunner.Terminate();
+		RequestEngineExit(TEXT("Exiting"));
+		FModuleManager::Get().UnloadModulesAtShutdown();
 	};
 
 	return TestRunner.RunCatchSession();
