LowLevelTests: Restored previous indentation in TestRunner.cpp

#rnx

diff --git a/Engine/Source/Developer/LowLevelTestsRunner/Private/TestRunner.cpp b/Engine/Source/Developer/LowLevelTestsRunner/Private/TestRunner.cpp
--- a/Engine/Source/Developer/LowLevelTestsRunner/Private/TestRunner.cpp
+++ b/Engine/Source/Developer/LowLevelTestsRunner/Private/TestRunner.cpp
@@ -25,232 +25,232 @@
 namespace UE::LowLevelTests
 {
 
-	static ITestRunner* GTestRunner;
+static ITestRunner* GTestRunner;
 
-	ITestRunner* ITestRunner::Get()
-	{
-		return GTestRunner;
-	}
-
-	ITestRunner::ITestRunner()
-	{
-		check(!GTestRunner);
-		GTestRunner = this;
-	}
-
-	ITestRunner::~ITestRunner()
-	{
-		check(GTestRunner == this);
-		GTestRunner = nullptr;
-	}
+ITestRunner* ITestRunner::Get()
+{
+	return GTestRunner;
+}
 
-	class FTestRunner final : public ITestRunner
-	{
-	public:
-		FTestRunner();
+ITestRunner::ITestRunner()
+{
+	check(!GTestRunner);
+	GTestRunner = this;
+}
 
-		void ParseCommandLine(TConstArrayView<const ANSICHAR*> Args);
+ITestRunner::~ITestRunner()
+{
+	check(GTestRunner == this);
+	GTestRunner = nullptr;
+}
 
-		void SleepOnInit() const;
+class FTestRunner final : public ITestRunner
+{
+public:
+	FTestRunner();
 
-		void GlobalSetup() const;
-		void GlobalTeardown() const;
-		void Terminate() const;
+	void ParseCommandLine(TConstArrayView<const ANSICHAR*> Args);
 
-		int32 RunCatchSession() const;
+	void SleepOnInit() const;
 
-		bool HasGlobalSetup() const final { return bGlobalSetup; }
-		bool HasLogOutput() const final { return bLogOutput || bDebugMode; }
-		bool IsDebugMode() const final { return bDebugMode; }
+	void GlobalSetup() const;
+	void GlobalTeardown() const;
+	void Terminate() const;
 
-	private:
-		static TArray<FName> GetGlobalModuleNames()
-		{
-			TArray<FName> ModuleNames;
-			FModuleManager::Get().FindModules(TEXT("*GlobalLowLevelTests"), ModuleNames);
-			return ModuleNames;
-		}
+	int32 RunCatchSession() const;
 
-		TArray<const ANSICHAR*> CatchArgs;
-		FStringBuilderBase ExtraArgs;
-		bool bGlobalSetup = true;
-		bool bLogOutput = false;
-		bool bDebugMode = false;
-		bool bMultiThreaded = false;
-		bool bWaitForInputToTerminate = false;
-		int32 SleepOnInitSeconds = 0;
-	};
+	bool HasGlobalSetup() const final { return bGlobalSetup; }
+	bool HasLogOutput() const final { return bLogOutput || bDebugMode; }
+	bool IsDebugMode() const final { return bDebugMode; }
 
-	FTestRunner::FTestRunner()
+private:
+	static TArray<FName> GetGlobalModuleNames()
 	{
-		// Start setting up the Game Thread.
-		GGameThreadId = FPlatformTLS::GetCurrentThreadId();
-		GIsGameThreadIdInitialized = true;
+		TArray<FName> ModuleNames;
+		FModuleManager::Get().FindModules(TEXT("*GlobalLowLevelTests"), ModuleNames);
+		return ModuleNames;
 	}
 
-	void FTestRunner::ParseCommandLine(TConstArrayView<const ANSICHAR*> Args)
+	TArray<const ANSICHAR*> CatchArgs;
+	FStringBuilderBase ExtraArgs;
+	bool bGlobalSetup = true;
+	bool bLogOutput = false;
+	bool bDebugMode = false;
+	bool bMultiThreaded = false;
+	bool bWaitForInputToTerminate = false;
+	int32 SleepOnInitSeconds = 0;
+};
+
+FTestRunner::FTestRunner()
+{
+	// Start setting up the Game Thread.
+	GGameThreadId = FPlatformTLS::GetCurrentThreadId();
+	GIsGameThreadIdInitialized = true;
+}
+
+void FTestRunner::ParseCommandLine(TConstArrayView<const ANSICHAR*> Args)
+{
+	bool bExtraArg = false;
+	for (FAnsiStringView Arg : Args)
 	{
-		bool bExtraArg = false;
-		for (FAnsiStringView Arg : Args)
+		if (bExtraArg)
 		{
-			if (bExtraArg)
+			if (const int32 SpaceIndex = String::FindFirstChar(Arg, ' '); SpaceIndex != INDEX_NONE)
 			{
-				if (const int32 SpaceIndex = String::FindFirstChar(Arg, ' '); SpaceIndex != INDEX_NONE)
-				{
-					if (const int32 EqualIndex = String::FindFirstChar(Arg, '='); EqualIndex != INDEX_NONE && EqualIndex < SpaceIndex)
-					{
-						ExtraArgs.Append(Arg.Left(EqualIndex + 1));
-						Arg.RightChopInline(EqualIndex + 1);
-					}
-					ExtraArgs.AppendChar('"').Append(Arg).AppendChar('"').AppendChar(' ');
-				}
-				else
+				if (const int32 EqualIndex = String::FindFirstChar(Arg, '='); EqualIndex != INDEX_NONE && EqualIndex < SpaceIndex)
 				{
-					ExtraArgs.Append(Arg).AppendChar(' ');
+					ExtraArgs.Append(Arg.Left(EqualIndex + 1));
+					Arg.RightChopInline(EqualIndex + 1);
 				}
-			}
-			else if (Arg == ANSITEXTVIEW("--"))
-			{
-				bExtraArg = true;
-			}
-			else if (Arg.StartsWith(ANSITEXTVIEW("--sleep=")))
-			{
-				LexFromString(SleepOnInitSeconds, WriteToString<16>(Arg.RightChop(8)).ToView());
-			}
-			else if (Arg == ANSITEXTVIEW("--global-setup"))
-			{
-				bGlobalSetup = true;
-			}
-			else if (Arg == ANSITEXTVIEW("--no-global-setup"))
-			{
-				bGlobalSetup = false;
-			}
-			else if (Arg == ANSITEXTVIEW("--log"))
-			{
-				bLogOutput = true;
-			}
-			else if (Arg == ANSITEXTVIEW("--no-log"))
-			{
-				bLogOutput = false;
-			}
-			else if (Arg == ANSITEXTVIEW("--debug"))
-			{
-				bDebugMode = true;
-			}
-			else if (Arg == ANSITEXTVIEW("--mt"))
-			{
-				bMultiThreaded = true;
-			}
-			else if (Arg == ANSITEXTVIEW("--no-mt"))
-			{
-				bMultiThreaded = false;
-			}
-			else if (Arg == ANSITEXTVIEW("--wait"))
-			{
-				bWaitForInputToTerminate = true;
-			}
-			else if (Arg == ANSITEXTVIEW("--no-wait"))
-			{
-				bWaitForInputToTerminate = true;
+				ExtraArgs.AppendChar('"').Append(Arg).AppendChar('"').AppendChar(' ');
 			}
 			else
 			{
-				CatchArgs.Add(Arg.GetData());
+				ExtraArgs.Append(Arg).AppendChar(' ');
 			}
 		}
-
-		// Break in the debugger on failed assertions when attached.
-		CatchArgs.Add("--break");
+		else if (Arg == ANSITEXTVIEW("--"))
+		{
+			bExtraArg = true;
+		}
+		else if (Arg.StartsWith(ANSITEXTVIEW("--sleep=")))
+		{
+			LexFromString(SleepOnInitSeconds, WriteToString<16>(Arg.RightChop(8)).ToView());
+		}
+		else if (Arg == ANSITEXTVIEW("--global-setup"))
+		{
+			bGlobalSetup = true;
+		}
+		else if (Arg == ANSITEXTVIEW("--no-global-setup"))
+		{
+			bGlobalSetup = false;
+		}
+		else if (Arg == ANSITEXTVIEW("--log"))
+		{
+			bLogOutput = true;
+		}
+		else if (Arg == ANSITEXTVIEW("--no-log"))
+		{
+			bLogOutput = false;
+		}
+		else if (Arg == ANSITEXTVIEW("--debug"))
+		{
+			bDebugMode = true;
+		}
+		else if (Arg == ANSITEXTVIEW("--mt"))
+		{
+			bMultiThreaded = true;
+		}
+		else if (Arg == ANSITEXTVIEW("--no-mt"))
+		{
+			bMultiThreaded = false;
+		}
+		else if (Arg == ANSITEXTVIEW("--wait"))
+		{
+			bWaitForInputToTerminate = true;
+		}
+		else if (Arg == ANSITEXTVIEW("--no-wait"))
+		{
+			bWaitForInputToTerminate = true;
+		}
+		else
+		{
+			CatchArgs.Add(Arg.GetData());
+		}
 	}
 
-	void FTestRunner::SleepOnInit() const
+	// Break in the debugger on failed assertions when attached.
+	CatchArgs.Add("--break");
+}
+
+void FTestRunner::SleepOnInit() const
+{
+	if (SleepOnInitSeconds)
 	{
-		if (SleepOnInitSeconds)
-		{
-			// Sleep to allow sync with Gauntlet.
-			FPlatformProcess::Sleep(SleepOnInitSeconds);
-		}
+		// Sleep to allow sync with Gauntlet.
+		FPlatformProcess::Sleep(SleepOnInitSeconds);
 	}
+}
 
-	void FTestRunner::GlobalSetup() const
+void FTestRunner::GlobalSetup() const
+{
+	if (!bGlobalSetup)
 	{
-		if (!bGlobalSetup)
-		{
-			return;
-		}
+		return;
+	}
 
-		FCommandLine::Set(*ExtraArgs);
+	FCommandLine::Set(*ExtraArgs);
 
-		// Finish setting up the Game Thread, which requires the command line.
-		FPlatformProcess::SetThreadAffinityMask(FPlatformAffinity::GetMainGameMask());
-		FPlatformProcess::SetupGameThread();
+	// Finish setting up the Game Thread, which requires the command line.
+	FPlatformProcess::SetThreadAffinityMask(FPlatformAffinity::GetMainGameMask());
+	FPlatformProcess::SetupGameThread();
 
-		// Always set up GError to handle FatalError, failed assertions, and crashes and other fatal errors.
+	// Always set up GError to handle FatalError, failed assertions, and crashes and other fatal errors.
 #if WITH_APPLICATION_CORE
-		GError = FPlatformApplicationMisc::GetErrorOutputDevice();
+	GError = FPlatformApplicationMisc::GetErrorOutputDevice();
 #else
-		GError = FPlatformOutputDevices::GetError();
+	GError = FPlatformOutputDevices::GetError();
 #endif
 
-		if (bLogOutput || bDebugMode)
-		{
-			// Set up GWarn to handle Error, Warning, Display; but only when log output is enabled.
+	if (bLogOutput || bDebugMode)
+	{
+		// Set up GWarn to handle Error, Warning, Display; but only when log output is enabled.
 #if WITH_APPLICATION_CORE
-			GWarn = FPlatformApplicationMisc::GetFeedbackContext();
+		GWarn = FPlatformApplicationMisc::GetFeedbackContext();
 #else
-			GWarn = FPlatformOutputDevices::GetFeedbackContext();
+		GWarn = FPlatformOutputDevices::GetFeedbackContext();
 #endif
 
-			// Set up default output devices to handle Log, Verbose, VeryVerbose.
-			FPlatformOutputDevices::SetupOutputDevices();
+		// Set up default output devices to handle Log, Verbose, VeryVerbose.
+		FPlatformOutputDevices::SetupOutputDevices();
 
-			FLogSuppressionInterface::Get().ProcessConfigAndCommandLine();
-		}
-
-		for (FName ModuleName : GetGlobalModuleNames())
-		{
-			if (ILowLevelTestsModule* Module = FModuleManager::LoadModulePtr<ILowLevelTestsModule>(ModuleName))
-			{
-				Module->GlobalSetup();
-			}
-		}
+		FLogSuppressionInterface::Get().ProcessConfigAndCommandLine();
 	}
 
-	void FTestRunner::GlobalTeardown() const
+	for (FName ModuleName : GetGlobalModuleNames())
 	{
-		if (!bGlobalSetup)
-		{
-			return;
-		}
-
-		for (FName ModuleName : GetGlobalModuleNames())
+		if (ILowLevelTestsModule* Module = FModuleManager::LoadModulePtr<ILowLevelTestsModule>(ModuleName))
 		{
-			if (ILowLevelTestsModule* Module = FModuleManager::GetModulePtr<ILowLevelTestsModule>(ModuleName))
-			{
-				Module->GlobalTeardown();
-			}
+			Module->GlobalSetup();
 		}
+	}
+}
 
-		CleanupPlatform();
-
-		FCommandLine::Reset();
+void FTestRunner::GlobalTeardown() const
+{
+	if (!bGlobalSetup)
+	{
+		return;
 	}
 
-	void FTestRunner::Terminate() const
+	for (FName ModuleName : GetGlobalModuleNames())
 	{
-#if PLATFORM_DESKTOP
-		if (bWaitForInputToTerminate)
+		if (ILowLevelTestsModule* Module = FModuleManager::GetModulePtr<ILowLevelTestsModule>(ModuleName))
 		{
-			std::cout << "Press enter to exit..." << std::endl;
-			std::cin.ignore();
+			Module->GlobalTeardown();
 		}
-#endif
 	}
 
-	int32 FTestRunner::RunCatchSession() const
+	CleanupPlatform();
+
+	FCommandLine::Reset();
+}
+
+void FTestRunner::Terminate() const
+{
+#if PLATFORM_DESKTOP
+	if (bWaitForInputToTerminate)
 	{
-		return Catch::Session().run(CatchArgs.Num(), CatchArgs.GetData());
+		std::cout << "Press enter to exit..." << std::endl;
+		std::cin.ignore();
 	}
+#endif
+}
+
+int32 FTestRunner::RunCatchSession() const
+{
+	return Catch::Session().run(CatchArgs.Num(), CatchArgs.GetData());
+}
 
 } // UE::LowLevelTests
 
