Fix a crash on shutdown with LowLevelTests- some threads still may be accessing FCommandLine right as we're tearing down so make sure it is set properly.

diff --git a/Engine/Source/Developer/LowLevelTestsRunner/Private/TestRunner.cpp b/Engine/Source/Developer/LowLevelTestsRunner/Private/TestRunner.cpp
--- a/Engine/Source/Developer/LowLevelTestsRunner/Private/TestRunner.cpp
+++ b/Engine/Source/Developer/LowLevelTestsRunner/Private/TestRunner.cpp
@@ -247,8 +247,6 @@ void FTestRunner::GlobalTeardown() const
 	}
 
 	CleanupPlatform();
-
-	FCommandLine::Reset();
 }
 
 void FTestRunner::Terminate() const
