[Backout] - CL22010754
Original CL Desc
-----------------------------------------------------------------
(LLT) Set cached data directories via FPlatformFile::InitializeAfterProjectFilePath -> fixes code path cycle for Program type applications where no project path is defined.

Reverting as it causes UE-165118 (Insights log not being gathered on some platforms). The problem was that some systems rely on the cached paths being available very early during the launch sequence. Delaying the initialization of these paths leads to non-existent paths being used and traces not being written.

#rnx

diff --git a/Engine/Source/Developer/LowLevelTestsRunner/Private/TestCommon/CoreUtilities.cpp b/Engine/Source/Developer/LowLevelTestsRunner/Private/TestCommon/CoreUtilities.cpp
--- a/Engine/Source/Developer/LowLevelTestsRunner/Private/TestCommon/CoreUtilities.cpp
+++ b/Engine/Source/Developer/LowLevelTestsRunner/Private/TestCommon/CoreUtilities.cpp
@@ -2,9 +2,7 @@
 
 #include "CoreMinimal.h"
 
-#include "GenericPlatform/GenericPlatformFile.h"
 #include "HAL/PlatformOutputDevices.h"
-#include "HAL/PlatformFileManager.h"
 #include "HAL/PlatformProcess.h"
 #include "Logging/LogSuppressionInterface.h"
 #include "Misc/CommandLine.h"
@@ -106,10 +104,4 @@ void CleanupTaskGraphAndDependencies()
 void CleanupPlatform()
 {
 	FPlatformMisc::PlatformTearDown();
-}
-
-void InitPlatformFile()
-{
-	IPlatformFile& PlatformFile = FPlatformFileManager::Get().GetPlatformFile();
-	PlatformFile.InitializeAfterProjectFilePath();
 }
\ No newline at end of file
diff --git a/Engine/Source/Developer/LowLevelTestsRunner/Private/TestCommon/Initialization.cpp b/Engine/Source/Developer/LowLevelTestsRunner/Private/TestCommon/Initialization.cpp
--- a/Engine/Source/Developer/LowLevelTestsRunner/Private/TestCommon/Initialization.cpp
+++ b/Engine/Source/Developer/LowLevelTestsRunner/Private/TestCommon/Initialization.cpp
@@ -39,7 +39,6 @@ void InitAll(bool bAllowLogging, bool bMultithreaded)
 	InitForWithEditorOnlyData();
 	InitEditor();
 #endif // WITH_EDITOR
-	InitPlatformFile();
 #if WITH_COREUOBJECT
 	InitCoreUObject();
 #endif
diff --git a/Engine/Source/Developer/LowLevelTestsRunner/Public/TestCommon/CoreUtilities.h b/Engine/Source/Developer/LowLevelTestsRunner/Public/TestCommon/CoreUtilities.h
--- a/Engine/Source/Developer/LowLevelTestsRunner/Public/TestCommon/CoreUtilities.h
+++ b/Engine/Source/Developer/LowLevelTestsRunner/Public/TestCommon/CoreUtilities.h
@@ -29,6 +29,4 @@ void InitTaskGraphAndDependencies(bool MultiThreaded = false);
 
 void CleanupTaskGraphAndDependencies();
 
-void CleanupPlatform();
-
-void InitPlatformFile();
\ No newline at end of file
+void CleanupPlatform();
\ No newline at end of file
diff --git a/Engine/Source/Programs/LowLevelTests/Tests/LowLevelTestsModule.cpp b/Engine/Source/Programs/LowLevelTests/Tests/LowLevelTestsModule.cpp
--- a/Engine/Source/Programs/LowLevelTests/Tests/LowLevelTestsModule.cpp
+++ b/Engine/Source/Programs/LowLevelTests/Tests/LowLevelTestsModule.cpp
@@ -18,7 +18,6 @@ void FGlobalLLTModule::GlobalSetup()
 {
 	InitTaskGraph();
 	InitThreadPool(true);
-	InitPlatformFile();
 #if WITH_COREUOBJECT
 	InitCoreUObject();
 #endif
