[Core] added support to lowlevel tests validating that `ensure` and `check` are executed

added  a tls for an ensure handler callback to let the reporting log of ensures to be overridden

the following macros are supported in lowlevel tests
REQUIRE_ENSURE
CHECK_ENSURE
REQUIRE_CHECK

diff --git a/Engine/Source/Developer/LowLevelTestsRunner/Private/CheckScope.cpp b/Engine/Source/Developer/LowLevelTestsRunner/Private/CheckScope.cpp
new file mode 100644
--- /dev/null
+++ b/Engine/Source/Developer/LowLevelTestsRunner/Private/CheckScope.cpp
@@ -0,0 +1,90 @@
+// Copyright Epic Games, Inc. All Rights Reserved.
+
+#include "LowLevelTestsRunner/CheckScope.h"
+#include "Containers/StringConv.h"
+#include "CoreGlobals.h"
+#include "Misc/OutputDeviceError.h"
+
+namespace UE::LowLevelTests
+{
+
+class FCheckScopeOutputDeviceError : public FOutputDeviceError
+{
+public:
+	FCheckScopeOutputDeviceError(FOutputDeviceError* Error, const ANSICHAR* Msg)
+		: DeviceError(Error)
+		, ExpectedMsg(Msg)
+		, Count(0)
+	{
+	}
+
+	virtual void HandleError() override
+	{
+		if (DeviceError)
+		{
+			DeviceError->HandleError();
+		}
+	}
+
+	void Serialize(const TCHAR* V, ELogVerbosity::Type Verbosity, const FName& Category) override
+	{
+		if (ExpectedMsg)
+		{
+			auto Str = StringCast<TCHAR>(ExpectedMsg);
+			auto Found = FPlatformString::Strstr(V, Str.Get());
+			if (Found)
+			{
+				++Count;
+			}
+		}
+		else
+		{
+			++Count;
+		}
+	}
+
+	void SetDeviceError(FOutputDeviceError* Error)
+	{
+		DeviceError = Error;
+	}
+
+	FOutputDeviceError* GetDeviceError() const
+	{
+		return DeviceError;
+	}
+
+	int GetCount() { return Count; }
+
+private:
+	FOutputDeviceError* DeviceError;
+	const ANSICHAR* ExpectedMsg;
+	int Count;
+};
+
+FCheckScope::FCheckScope(const ANSICHAR* Msg)
+	: DeviceError(new FCheckScopeOutputDeviceError(GError, Msg))
+	, bIgnoreDebugger(GIgnoreDebugger)
+{
+	GIgnoreDebugger = true;
+	GError = DeviceError;
+}
+
+FCheckScope::FCheckScope()
+	: FCheckScope(nullptr)	
+{
+}
+
+FCheckScope::~FCheckScope()
+{
+	GIgnoreDebugger = bIgnoreDebugger;
+	GError = DeviceError->GetDeviceError();
+	delete DeviceError;
+}
+
+int
+FCheckScope::GetCount()
+{
+	return DeviceError->GetCount();
+}
+
+}
\ No newline at end of file
diff --git a/Engine/Source/Developer/LowLevelTestsRunner/Private/EnsureScope.cpp b/Engine/Source/Developer/LowLevelTestsRunner/Private/EnsureScope.cpp
new file mode 100644
--- /dev/null
+++ b/Engine/Source/Developer/LowLevelTestsRunner/Private/EnsureScope.cpp
@@ -0,0 +1,42 @@
+// Copyright Epic Games, Inc. All Rights Reserved.
+
+#include "LowLevelTestsRunner/EnsureScope.h"
+#include "Containers/StringConv.h"
+
+namespace UE::LowLevelTests
+{
+
+FEnsureScope::FEnsureScope()
+	: FEnsureScope([](const FEnsureHandlerArgs&) { return true; })
+{
+}
+
+
+FEnsureScope::FEnsureScope(const ANSICHAR* ExpectedMsg)
+	: FEnsureScope([ExpectedMsg](const FEnsureHandlerArgs& Args)
+		{
+			return FPlatformString::Stricmp(ExpectedMsg, Args.Message) == 0;
+		})
+{
+}
+
+FEnsureScope::FEnsureScope(TFunction<bool(const FEnsureHandlerArgs& Args)> EnsureFunc)
+	: Count(0)
+{
+	OldHandler = SetEnsureHandler([this, EnsureFunc](const FEnsureHandlerArgs& Args) -> bool
+		{
+			bool Handled = EnsureFunc(Args);
+			if (Handled)
+			{
+				++Count;
+			}
+			return Handled;
+		});
+}
+
+FEnsureScope::~FEnsureScope()
+{
+	SetEnsureHandler(OldHandler);
+}
+
+}
\ No newline at end of file
diff --git a/Engine/Source/Developer/LowLevelTestsRunner/Private/TestRunner.cpp b/Engine/Source/Developer/LowLevelTestsRunner/Private/TestRunner.cpp
--- a/Engine/Source/Developer/LowLevelTestsRunner/Private/TestRunner.cpp
+++ b/Engine/Source/Developer/LowLevelTestsRunner/Private/TestRunner.cpp
@@ -13,12 +13,15 @@
 #include "String/Find.h"
 #include "String/LexFromString.h"
 #include "TestCommon/CoreUtilities.h"
+#include "TestRunnerOutputDeviceError.h"
 
 #if WITH_APPLICATION_CORE
 #include "HAL/PlatformApplicationMisc.h"
 #endif
 
+#include "Misc/CoreDelegates.h"
 #include <catch2/catch_session.hpp>
+#include <catch2/internal/catch_assertion_handler.hpp>
 
 #include <iostream>
 
@@ -53,7 +56,7 @@ public:
 
 	void SleepOnInit() const;
 
-	void GlobalSetup() const;
+	void GlobalSetup();
 	void GlobalTeardown() const;
 	void Terminate() const;
 
@@ -73,6 +76,7 @@ private:
 
 	TArray<const ANSICHAR*> CatchArgs;
 	FStringBuilderBase ExtraArgs;
+	FTestRunnerOutputDeviceError ErrorOutputDevice;
 	bool bGlobalSetup = true;
 	bool bLogOutput = false;
 	bool bDebugMode = false;
@@ -177,7 +181,7 @@ void FTestRunner::SleepOnInit() const
 	}
 }
 
-void FTestRunner::GlobalSetup() const
+void FTestRunner::GlobalSetup()
 {
 	if (bAttachToDebugger)
 	{
@@ -205,7 +209,19 @@ void FTestRunner::GlobalSetup() const
 	GError = FPlatformApplicationMisc::GetErrorOutputDevice();
 #else
 	GError = FPlatformOutputDevices::GetError();
+	ErrorOutputDevice.SetDeviceError(GError);
+	GError = &ErrorOutputDevice;
 #endif
+	
+	//forward unhandled `ensure` to catch to force tests to fail. test will continue to execute
+	//this does bypass the error reporting, crash reporter and etc
+	FCoreDelegates::OnHandleSystemEnsure.AddLambda([this]()
+		{
+			FString Error = GErrorHist;
+			Catch::AssertionInfo info{ "", CATCH_INTERNAL_LINEINFO, "", Catch::ResultDisposition::Normal };
+			Catch::AssertionReaction reaction;
+			Catch::getResultCapture().handleMessage(info, Catch::ResultWas::ExplicitFailure, StringCast<ANSICHAR>(*Error).Get(), reaction);
+		});
 
 	if (bLogOutput || bDebugMode)
 	{
@@ -237,6 +253,12 @@ void FTestRunner::GlobalTeardown() const
 	{
 		return;
 	}
+	
+	//only set the GError back if it was replaced
+	if (GError == &ErrorOutputDevice)
+	{
+		GError = ErrorOutputDevice.GetDeviceError();
+	}
 
 	for (FName ModuleName : GetGlobalModuleNames())
 	{
@@ -298,5 +320,6 @@ int RunTests(int32 ArgC, const ANSICHAR* ArgV[])
 		FModuleManager::Get().UnloadModulesAtShutdown();
 	};
 
-	return TestRunner.RunCatchSession();
+	int CatchReturn = TestRunner.RunCatchSession();
+	return CatchReturn;
 }
diff --git a/Engine/Source/Developer/LowLevelTestsRunner/Private/TestRunnerOutputDeviceError.cpp b/Engine/Source/Developer/LowLevelTestsRunner/Private/TestRunnerOutputDeviceError.cpp
new file mode 100644
--- /dev/null
+++ b/Engine/Source/Developer/LowLevelTestsRunner/Private/TestRunnerOutputDeviceError.cpp
@@ -0,0 +1,20 @@
+// Copyright Epic Games, Inc. All Rights Reserved.
+
+#include "TestRunnerOutputDeviceError.h"
+#include "Containers/StringConv.h"
+#include <catch2/internal/catch_assertion_handler.hpp>
+
+namespace UE::LowLevelTests
+{
+
+//report error back to catch otherwise failed `check`'s don't fail the tests
+//this method is not execute if the debugger is attached
+void FTestRunnerOutputDeviceError::Serialize(const TCHAR* V, ELogVerbosity::Type Verbosity, const FName& Category)
+{
+	auto AnsiMessage = StringCast<ANSICHAR>(V);
+	Catch::AssertionHandler catchAssertionHandler("", CATCH_INTERNAL_LINEINFO, AnsiMessage.Get(), Catch::ResultDisposition::Normal);
+	catchAssertionHandler.handleMessage(Catch::ResultWas::ExplicitFailure, AnsiMessage.Get());
+	catchAssertionHandler.complete();
+}
+
+}
\ No newline at end of file
diff --git a/Engine/Source/Developer/LowLevelTestsRunner/Private/TestRunnerOutputDeviceError.h b/Engine/Source/Developer/LowLevelTestsRunner/Private/TestRunnerOutputDeviceError.h
new file mode 100644
--- /dev/null
+++ b/Engine/Source/Developer/LowLevelTestsRunner/Private/TestRunnerOutputDeviceError.h
@@ -0,0 +1,47 @@
+// Copyright Epic Games, Inc. All Rights Reserved.
+
+#pragma once
+#include "Misc/OutputDeviceError.h"
+
+namespace UE::LowLevelTests
+{
+
+///@brief Used to replace GError so that errors are capture during execution of tests
+class FTestRunnerOutputDeviceError : public FOutputDeviceError
+{
+public:
+	explicit FTestRunnerOutputDeviceError(FOutputDeviceError* Error)
+		: DeviceError(Error)
+	{
+	}
+
+	FTestRunnerOutputDeviceError()
+		: DeviceError(nullptr)
+	{
+
+	}
+
+	virtual void HandleError() override
+	{
+		if (DeviceError)
+		{
+			DeviceError->HandleError();
+		}
+	}
+
+	void Serialize(const TCHAR* V, ELogVerbosity::Type Verbosity, const FName& Category) override;
+
+	void SetDeviceError(FOutputDeviceError* Error)
+	{
+		DeviceError = Error;
+	}
+
+	FOutputDeviceError* GetDeviceError() const
+	{
+		return DeviceError;
+	}
+
+private:
+	FOutputDeviceError* DeviceError;
+};
+}
\ No newline at end of file
diff --git a/Engine/Source/Developer/LowLevelTestsRunner/Public/LowLevelTestsRunner/CheckScope.h b/Engine/Source/Developer/LowLevelTestsRunner/Public/LowLevelTestsRunner/CheckScope.h
new file mode 100644
--- /dev/null
+++ b/Engine/Source/Developer/LowLevelTestsRunner/Public/LowLevelTestsRunner/CheckScope.h
@@ -0,0 +1,27 @@
+// Copyright Epic Games, Inc. All Rights Reserved.
+
+#pragma once
+#include "HAL/Platform.h"
+
+namespace UE::LowLevelTests
+{
+
+class FCheckScopeOutputDeviceError;
+
+///@brief Scope that captures failed `check` calls
+struct FCheckScope
+{
+	FCheckScope();
+
+	///@brief Only catches failed `check` calls that contain Msg
+	///@param Msg string to look for in the `check` message
+	FCheckScope(const ANSICHAR* Msg);
+	~FCheckScope();
+
+	int GetCount();
+private:
+	FCheckScopeOutputDeviceError* DeviceError;
+	bool bIgnoreDebugger;
+}; 
+
+}
\ No newline at end of file
diff --git a/Engine/Source/Developer/LowLevelTestsRunner/Public/LowLevelTestsRunner/EnsureScope.h b/Engine/Source/Developer/LowLevelTestsRunner/Public/LowLevelTestsRunner/EnsureScope.h
new file mode 100644
--- /dev/null
+++ b/Engine/Source/Developer/LowLevelTestsRunner/Public/LowLevelTestsRunner/EnsureScope.h
@@ -0,0 +1,32 @@
+// Copyright Epic Games, Inc. All Rights Reserved.
+
+#pragma once
+#include "Templates/Function.h"
+
+struct FEnsureHandlerArgs;
+
+namespace UE::LowLevelTests
+{
+
+///@brief Scope that captures failed `ensure` calls
+struct FEnsureScope
+{
+	FEnsureScope();
+
+	///@brief Only captures failed `ensure` calls with the exact message 
+	///@param Msg string to compare against failed `ensure` message
+	FEnsureScope(const ANSICHAR* Msg);
+
+	///@brief Captures failed `ensure` calls that return true from callback
+	///@param EnsureFunc lambda called for each failed `ensure` return true to handle the error
+	FEnsureScope(TFunction<bool(const FEnsureHandlerArgs&)> EnsureFunc);
+	~FEnsureScope();
+
+	int GetCount() { return Count; }
+
+private:
+	TFunction<bool(const FEnsureHandlerArgs&)> OldHandler;
+	int Count;
+};
+
+}
\ No newline at end of file
diff --git a/Engine/Source/Developer/LowLevelTestsRunner/Public/TestMacros/Assertions.h b/Engine/Source/Developer/LowLevelTestsRunner/Public/TestMacros/Assertions.h
--- a/Engine/Source/Developer/LowLevelTestsRunner/Public/TestMacros/Assertions.h
+++ b/Engine/Source/Developer/LowLevelTestsRunner/Public/TestMacros/Assertions.h
@@ -2,39 +2,76 @@
 #pragma once
 
 #include "Misc/AssertionMacros.h"
+#include "LowLevelTestsRunner/EnsureScope.h"
+#include "LowLevelTestsRunner/CheckScope.h"
+
+//requires that an UE `ensure` fails in this call
+#define REQUIRE_ENSURE(...) INTERNAL_UE_ENSURE( "REQUIRE_ENSURE", Catch::ResultDisposition::Normal, __VA_ARGS__ )
+
+//requires that an UE `ensure` fails with a message that matches the supplied message
+#define REQUIRE_ENSURE_MSG(msg, ...) INTERNAL_UE_ENSURE_MSG(msg, "REQUIRE_ENSURE", Catch::ResultDisposition::Normal, __VA_ARGS__ )
+
+//checks that an UE `ensure` fails in this call
+#define CHECK_ENSURE(...) INTERNAL_UE_ENSURE( "CHECK_ENSURE", Catch::ResultDisposition::ContinueOnFailure, __VA_ARGS__ )
+
+//checks that an UE `ensure` fails with a message that matches the supplied message
+#define CHECK_ENSURE_MSG(msg, ...) INTERNAL_UE_ENSURE_MSG(msg, "CHECK_ENSURE", Catch::ResultDisposition::ContinueOnFailure, __VA_ARGS__ )
+
+//requires that a UE `check` fails in this call
+#define REQUIRE_CHECK(...) INTERNAL_UE_CHECK( "REQUIRE_ENSURE", Catch::ResultDisposition::Normal, __VA_ARGS__ )
+
+//requires that a UE `check` fails in this call containg the supplies message
+#define REQUIRE_CHECK_MSG(msg, ...) INTERNAL_UE_CHECK_MSG(msg, "REQUIRE_ENSURE", Catch::ResultDisposition::Normal, __VA_ARGS__ )
+
+#define INTERNAL_UE_ENSURE( macroName, resultDisposition, ... ) \
+	do { \
+		Catch::AssertionHandler catchAssertionHandler( macroName##_catch_sr, CATCH_INTERNAL_LINEINFO, CATCH_INTERNAL_STRINGIFY(__VA_ARGS__), resultDisposition ); \
+		INTERNAL_CATCH_TRY { \
+			::UE::LowLevelTests::FEnsureScope scope; \
+			static_cast<void>(__VA_ARGS__); \
+			bool bEncounteredEnsure = scope.GetCount() > 0; \
+			if (!bEncounteredEnsure) \
+				catchAssertionHandler.handleMessage( Catch::ResultWas::ExplicitFailure ,"Expected failure of `ensure` not recieved."); \
+		} INTERNAL_CATCH_CATCH( catchAssertionHandler ) \
+		INTERNAL_CATCH_REACT( catchAssertionHandler ) \
+	} while(false) \
+
+#define INTERNAL_UE_ENSURE_MSG(msg, macroName, resultDisposition, ... ) \
+	do { \
+		Catch::AssertionHandler catchAssertionHandler( macroName##_catch_sr, CATCH_INTERNAL_LINEINFO, CATCH_INTERNAL_STRINGIFY(__VA_ARGS__), resultDisposition ); \
+		INTERNAL_CATCH_TRY { \
+			::UE::LowLevelTests::FEnsureScope scope(msg); \
+			static_cast<void>(__VA_ARGS__); \
+			bool bEncounteredEnsure = scope.GetCount() > 0; \
+			if (!bEncounteredEnsure) \
+				catchAssertionHandler.handleMessage( Catch::ResultWas::ExplicitFailure, ( Catch::MessageStream() << "Expected failure of `ensure` with message '" << msg << "' not recieved" + ::Catch::StreamEndStop() ).m_stream.str() ); \
+		} INTERNAL_CATCH_CATCH( catchAssertionHandler ) \
+		INTERNAL_CATCH_REACT( catchAssertionHandler ) \
+	} while(false) \
+
+#define INTERNAL_UE_CHECK(macroName, resultDisposition, ... ) \
+	do { \
+		Catch::AssertionHandler catchAssertionHandler( macroName##_catch_sr, CATCH_INTERNAL_LINEINFO, CATCH_INTERNAL_STRINGIFY(__VA_ARGS__), resultDisposition ); \
+		INTERNAL_CATCH_TRY { \
+			::UE::LowLevelTests::FCheckScope scope; \
+			__VA_ARGS__; \
+			bool bEncounteredEnsure = scope.GetCount() > 0; \
+			if (!bEncounteredEnsure) \
+				catchAssertionHandler.handleMessage( Catch::ResultWas::ExplicitFailure, "Expected failure of `check` not recieved"); \
+		} INTERNAL_CATCH_CATCH( catchAssertionHandler ) \
+		INTERNAL_CATCH_REACT( catchAssertionHandler ) \
+	} while(false) \
+
+#define INTERNAL_UE_CHECK_MSG(msg, macroName, resultDisposition, ... ) \
+	do { \
+		Catch::AssertionHandler catchAssertionHandler( macroName##_catch_sr, CATCH_INTERNAL_LINEINFO, CATCH_INTERNAL_STRINGIFY(__VA_ARGS__), resultDisposition ); \
+		INTERNAL_CATCH_TRY { \
+			::UE::LowLevelTests::FCheckScope scope(msg); \
+			__VA_ARGS__; \
+			bool bEncounteredEnsure = scope.GetCount() > 0; \
+			if (!bEncounteredEnsure) \
+				catchAssertionHandler.handleMessage( Catch::ResultWas::ExplicitFailure, ( Catch::MessageStream() << "Expected failure of `check` containing message '" << msg << "' not recieved" + ::Catch::StreamEndStop() ).m_stream.str() ); \
+		} INTERNAL_CATCH_CATCH( catchAssertionHandler ) \
+		INTERNAL_CATCH_REACT( catchAssertionHandler ) \
+	} while(false) \
 
-#if defined(CATCH_CONFIG_PREFIX_ALL) && !defined(CATCH_CONFIG_DISABLE)
-    #define CATCH_REQUIRE_ENSURE( ... ) INTERNAL_CATCH_ENSURE( "CATCH_REQUIRE_ENSURE", Catch::ResultDisposition::Normal, __VA_ARGS__ )
-    #define CATCH_REQUIRE_NOENSURE( ... ) INTERNAL_CATCH_NO_ENSURE( "CATCH_REQUIRE_NOENSURE", Catch::ResultDisposition::Normal, __VA_ARGS__ )
-    #define CATCH_CHECK_ENSURE( ... ) INTERNAL_CATCH_ENSURE( "CATCH_CHECK_ENSURE", Catch::ResultDisposition::ContinueOnFailure, __VA_ARGS__ )
-    #define CATCH_CHECK_NOENSURE( ... ) INTERNAL_CATCH_NO_ENSURE( "CATCH_CHECK_NOENSURE", Catch::ResultDisposition::ContinueOnFailure, __VA_ARGS__ )
-#else
-    #define REQUIRE_ENSURE( ... ) INTERNAL_CATCH_ENSURE( "REQUIRE_ENSURE", Catch::ResultDisposition::Normal, __VA_ARGS__ )
-    #define REQUIRE_NOENSURE( ... ) INTERNAL_CATCH_NO_ENSURE( "REQUIRE_NOENSURE", Catch::ResultDisposition::Normal, __VA_ARGS__ )
-    #define CHECK_ENSURE( ... ) INTERNAL_CATCH_ENSURE( "CHECK_ENSURE", Catch::ResultDisposition::ContinueOnFailure, __VA_ARGS__ )
-    #define CHECK_NOENSURE( ... ) INTERNAL_CATCH_NO_ENSURE( "CHECK_NOENSURE", Catch::ResultDisposition::ContinueOnFailure, __VA_ARGS__ )
-#endif
-
-#define INTERNAL_CATCH_ENSURE( macroName, resultDisposition, ... ) \
- do { \
-		SIZE_T NumEnsureFailuresBefore = FDebug::GetNumEnsureFailures(); \
-        Catch::AssertionHandler catchAssertionHandler( macroName##_catch_sr, CATCH_INTERNAL_LINEINFO, CATCH_INTERNAL_STRINGIFY(__VA_ARGS__), resultDisposition ); \
-        INTERNAL_CATCH_TRY { \
-            static_cast<void>(__VA_ARGS__); \
-            bool bEncounteredEnsure = (FDebug::GetNumEnsureFailures() - NumEnsureFailuresBefore) > 0; \
-			catchAssertionHandler.handleExpr( Catch::Decomposer() <= bEncounteredEnsure ); \
-        } INTERNAL_CATCH_CATCH( catchAssertionHandler ) \
-        INTERNAL_CATCH_REACT( catchAssertionHandler ) \
-    } while( false )
-
-#define INTERNAL_CATCH_NO_ENSURE( macroName, resultDisposition, ... ) \
- do { \
-		SIZE_T NumEnsureFailuresBefore = FDebug::GetNumEnsureFailures(); \
-        Catch::AssertionHandler catchAssertionHandler( macroName##_catch_sr, CATCH_INTERNAL_LINEINFO, CATCH_INTERNAL_STRINGIFY(__VA_ARGS__), resultDisposition ); \
-        INTERNAL_CATCH_TRY { \
-            static_cast<void>(__VA_ARGS__); \
-            bool bEncounteredEnsure = (FDebug::GetNumEnsureFailures() - NumEnsureFailuresBefore) > 0; \
-			catchAssertionHandler.handleExpr( Catch::Decomposer() <= !(bEncounteredEnsure) ); \
-        } INTERNAL_CATCH_CATCH( catchAssertionHandler ) \
-        INTERNAL_CATCH_REACT( catchAssertionHandler ) \
-    } while( false )
diff --git a/Engine/Source/Runtime/Core/Private/Misc/AssertionMacros.cpp b/Engine/Source/Runtime/Core/Private/Misc/AssertionMacros.cpp
--- a/Engine/Source/Runtime/Core/Private/Misc/AssertionMacros.cpp
+++ b/Engine/Source/Runtime/Core/Private/Misc/AssertionMacros.cpp
@@ -361,6 +361,20 @@ void FDebug::LogAssertFailedMessageImplV(const ANSICHAR* Expr, const ANSICHAR* F
 	}
 }
 
+thread_local TFunction<bool(const FEnsureHandlerArgs& Args)> EnsureHandler = nullptr;
+
+TFunction<bool(const FEnsureHandlerArgs& Args)> SetEnsureHandler(TFunction<bool(const FEnsureHandlerArgs& Args)> Handler)
+{
+	TFunction<bool(const FEnsureHandlerArgs& Args)> OldHandler = EnsureHandler;
+	EnsureHandler = MoveTemp(Handler);
+	return OldHandler;
+}
+
+TFunction<bool(const FEnsureHandlerArgs& Args)> GetEnsureHandler()
+{
+	return EnsureHandler;
+}
+
 /**
  * Called when an 'ensure' assertion fails; gathers stack data and generates and error report.
  *
@@ -372,6 +386,10 @@ void FDebug::LogAssertFailedMessageImplV(const ANSICHAR* Expr, const ANSICHAR* F
  */
 FORCENOINLINE void FDebug::EnsureFailed(const ANSICHAR* Expr, const ANSICHAR* File, int32 Line, void* ProgramCounter, const TCHAR* Msg)
 {
+	if (UNLIKELY(EnsureHandler && EnsureHandler({ Expr, Msg })))
+	{
+		return;
+	}
 	FTempCommandLineScope TempCommandLine;
 
 	// if time isn't ready yet, we better not continue
diff --git a/Engine/Source/Runtime/Core/Public/Misc/AssertionMacros.h b/Engine/Source/Runtime/Core/Public/Misc/AssertionMacros.h
--- a/Engine/Source/Runtime/Core/Public/Misc/AssertionMacros.h
+++ b/Engine/Source/Runtime/Core/Public/Misc/AssertionMacros.h
@@ -37,6 +37,23 @@ namespace ELogVerbosity
  */
 extern "C" CORE_API void PrintScriptCallstack();
 
+
+struct FEnsureHandlerArgs
+{
+	const ANSICHAR* Expression;
+	const TCHAR* Message;
+};
+
+template <typename FuncType> class TFunction;
+
+/** 
+ * Sets the thead_local ensure handler callback
+ * returning true skips the regular ensure handling logic
+ * returns the previous handler which maybe nullptr
+ **/
+CORE_API TFunction<bool(const FEnsureHandlerArgs& Args)> SetEnsureHandler(TFunction<bool(const FEnsureHandlerArgs& Args)> EnsureHandler);
+CORE_API TFunction<bool(const FEnsureHandlerArgs& Args)> GetEnsureHandler();
+
 /**
  * FDebug
  * These functions offer debugging and diagnostic functionality and its presence 
diff --git a/Engine/Source/Runtime/Core/Tests/Misc/AssertionMacrosTest.cpp b/Engine/Source/Runtime/Core/Tests/Misc/AssertionMacrosTest.cpp
--- a/Engine/Source/Runtime/Core/Tests/Misc/AssertionMacrosTest.cpp
+++ b/Engine/Source/Runtime/Core/Tests/Misc/AssertionMacrosTest.cpp
@@ -1,5 +1,7 @@
 // Copyright Epic Games, Inc. All Rights Reserved.
 
+#include "Misc/Build.h" 
+
 #if WITH_LOW_LEVEL_TESTS && DO_ENSURE
 
 #include "Misc/AssertionMacros.h"
@@ -8,116 +10,116 @@
 #include "TestHarness.h"
 #include "TestMacros/Assertions.h"
 
-static void OneEnsureNonFailed();
-static void OneEnsureNonFailedDifferent();
-static void OneEnsureFailed();
-static void OneEnsureFailedDifferent();
-static void OneEnsureFailedDifferentAgain();
-static void MultipleEnsuresNonFailed();
-static void MultipleEnsuresFailed();
-static void MultipleEnsuresFailedDifferent();
+//since ensure uses a static bool to suppress multiple calls
+//these small functions are needed to scope that bool
+static void EnsureFailed()
+{
+	ensure(false);
+}
 
-/**
- * @brief Ensure test case, tests Catch2-style macros that check or require whether at least one ensure was triggered or not.
- * Note that each different ensure expression is counted once, if the same ensure fails twice it's only counted the first time it fails.
- */
-TEST_CASE("Core::Misc::AssertionMacros", "[Core][Misc][AssertionMacros][Ensure]")
+static void EnsureFailAlways()
 {
-	TGuardValue<bool> IgnoreDebugger(GIgnoreDebugger, true);
-	SECTION("Ensure not triggered with REQUIRE_NOENSURE")
+	ensureAlways(false);
+}
+
+static int EnsureContinues(int Value)
+{
+	if (ensure(Value == 1))
 	{
-		REQUIRE_NOENSURE(OneEnsureNonFailed());
+		return Value + 1;
 	}
+	return Value + 2;
+}
 
-	/**
-	 * Some sections are disabled until assertion messages and/or callstacks are skipped in a thread-local manner.
-	*/
+static int EnsureMultipleFailures(int Value)
+{
+	ensure(Value == 1);
+	ensure(Value == 5);
+	ensure(Value == 2);
+	return Value + 1;
+}
 
-	DISABLED_SECTION("Ensure triggered with REQUIRE_ENSURE")
+static int EnsureForCheck(int Value)
+{
+	if (ensure(Value == 1))
 	{
-		REQUIRE_ENSURE(OneEnsureFailed());
-	};
+		return Value + 1;
+	}
+	return Value + 2;
+}
 
-	SECTION("Ensure not triggered with multiple CHECK_NOENSURE")
+static int EnsureAlwaysMsg(int Value)
+{
+	if (ensureAlwaysMsgf(Value == 1, TEXT("expected value of 1")))
 	{
-		SIZE_T EnsuresBefore = FDebug::GetNumEnsureFailures();
-		CHECK_NOENSURE(OneEnsureNonFailed());
-		CHECK_NOENSURE(OneEnsureNonFailedDifferent());
-		CHECK(FDebug::GetNumEnsureFailures() == EnsuresBefore);
+		return Value + 1;
 	}
+	return Value + 2;
+}
 
-	DISABLED_SECTION("Ensure triggered multiple CHECK_ENSURE")
+/**
+ * @brief Ensure test case, tests Catch2-style macros that check or require whether at least one ensure was triggered or not.
+ * Note that each different ensure expression is counted once, if the same ensure fails twice it's only counted the first time it fails.
+ */
+TEST_CASE("Core::Misc::TestEnsure", "[Core][Misc][AssertionMacros][Ensure]")
+{
+	//ignore the debugger as all the breaks make for annoying debugging experience
+	TGuardValue<bool> IgnoreDebugger(GIgnoreDebugger, true);
 	{
-		SIZE_T EnsuresBefore = FDebug::GetNumEnsureFailures();
-		CHECK_ENSURE(OneEnsureFailedDifferent());
-		CHECK_ENSURE(OneEnsureFailedDifferentAgain());
-		CHECK(FDebug::GetNumEnsureFailures() == EnsuresBefore + 2);
-	};
+		REQUIRE_ENSURE(EnsureFailed());
+		EnsureFailed(); //ensure already fired should not fire again
+	}
 
-	SECTION("Multiple ensures not triggered with one REQUIRE_NOENSURE")
 	{
-		REQUIRE_NOENSURE(MultipleEnsuresNonFailed());
+		//validate that ensure continue on failure
+		//validate that non failure works
+		int result = 0;
+		result = EnsureContinues(1);
+		REQUIRE(result == 2);
+
+		REQUIRE_ENSURE(result = EnsureContinues(0));
+		REQUIRE(result == 2);
+
+		//second failure should not cause an error
+		result = EnsureContinues(2);
+		REQUIRE(result == 4);
 	}
 
-	DISABLED_SECTION("At least one of multiple ensures triggered with REQUIRE_ENSURE")
 	{
-		REQUIRE_ENSURE(MultipleEnsuresFailed());
-	};
+		//validate that ensureAlways raises errors
+		REQUIRE_ENSURE(EnsureFailAlways());
+		REQUIRE_ENSURE(EnsureFailAlways());
+	}
 
-	SECTION("Multiple ensures not triggered with CHECK_NOENSURE")
 	{
-		CHECK_NOENSURE(MultipleEnsuresNonFailed());
+		//not much of test of CHECK_ENSURE as its hard to tell if the coninue part is working
+		int result = 0;
+		CHECK_ENSURE(result = EnsureForCheck(0));
+		REQUIRE(result == 2);
 	}
 	
-	DISABLED_SECTION("At least one of multiple ensures triggered with CHECK_ENSURE")
 	{
-		CHECK_ENSURE(MultipleEnsuresFailedDifferent());
-	};
-}
-
-static void OneEnsureNonFailed()
-{
-	ensure(true);
-}
-
-static void OneEnsureNonFailedDifferent()
-{
-	ensure(1 == 1);
-}
-
-static void OneEnsureFailed()
-{
-	ensure(false);
-}
-
-static void OneEnsureFailedDifferent()
-{
-	ensure(1 == 2);
-}
-
-static void OneEnsureFailedDifferentAgain()
-{
-	ensure(2 == 3);
-}
+		//multiple failed ensure should continue beecause its wrapped in a scope
+		int result = 0;
+		REQUIRE_ENSURE(result = EnsureMultipleFailures(7));
+		REQUIRE(result == 8);
+	}
 
-static void MultipleEnsuresNonFailed()
-{
-	ensure(2 == 2);
-	ensure(3 == 3);
+	{
+		int result = 0;
+		REQUIRE_ENSURE_MSG("expected value of 1", result = EnsureAlwaysMsg(3));
+		REQUIRE(result == 5);
+		CHECK_ENSURE_MSG("expected value of 1", result = EnsureAlwaysMsg(4));
+		REQUIRE(result == 6);
+	}
 }
 
-static void MultipleEnsuresFailed()
+TEST_CASE("Core::Misc::TestCheck", "[Core][Misc][AssertionMacros][Check]")
 {
-	ensure(2 == 2);
-	ensure(3 == 4);
-	ensure(4 == 4);
+	REQUIRE_CHECK(check(1 == 2));
+	REQUIRE_CHECK_MSG("1 == 2", checkf(1 == 2, TEXT("Error Message")));
+	REQUIRE_CHECK_MSG("Error Message", checkf(1 == 2, TEXT("Error Message")));
 }
 
-static void MultipleEnsuresFailedDifferent()
-{
-	ensure(2 == 2);
-	ensure(4 == 5);
-	ensure(3 == 3);
-}
 
 #endif
