[LowLevelTestRunner] added custom reporter to make parsing of xml easier for UnrealVS TestAdapter

diff --git a/Engine/Source/Developer/LowLevelTestsRunner/Private/UnrealReporter.cpp b/Engine/Source/Developer/LowLevelTestsRunner/Private/UnrealReporter.cpp
new file mode 100644
--- /dev/null
+++ b/Engine/Source/Developer/LowLevelTestsRunner/Private/UnrealReporter.cpp
@@ -0,0 +1,162 @@
+// Copyright Epic Games, Inc. All Rights Reserved.
+
+#include <catch2/reporters/catch_reporter_streaming_base.hpp>
+#include <catch2/catch_test_case_info.hpp>
+#include <catch2/reporters/catch_reporter_registrars.hpp>
+#include <catch2/internal/catch_xmlwriter.hpp>
+#include <catch2/internal/catch_textflow.hpp>
+#include "Misc/StringBuilder.h"
+
+#include <iostream>
+
+namespace UE::LowLevelTests
+{
+/** Simple reporter for the low level tests
+
+<testrun>
+	<testcase name="test::name">
+		<failure>
+			all error text 
+		</failure>
+		<result success="true/false" />
+	</testcase>
+</testrun>
+
+no tags, no sections. errors are aggregated together.
+**/
+class FUnrealReporter : public Catch::StreamingReporterBase
+{
+public:
+	using Catch::StreamingReporterBase::StreamingReporterBase;
+
+	FUnrealReporter(Catch::ReporterConfig&& Config)
+		: Catch::StreamingReporterBase(CATCH_MOVE(Config))
+		, Xml(m_stream)
+	{
+		m_preferences.shouldRedirectStdOut = true;
+		m_preferences.shouldReportAllAssertions = true;
+	}
+
+	void testRunStarting(Catch::TestRunInfo const& TestInfo)
+	{
+		StreamingReporterBase::testRunStarting(TestInfo);
+		Xml.startElement("testrun");
+	}
+
+	virtual void testRunEnded(Catch::TestRunStats const&) override
+	{
+		Xml.endElement();
+	}
+
+	void testCaseStarting(Catch::TestCaseInfo const& TestCaseInfo) override
+	{
+		Xml.startElement("testcase");
+		Xml.writeAttribute("name", TestCaseInfo.name);
+	}
+
+	virtual void testCaseEnded(Catch::TestCaseStats const& TestCaseStats) override
+	{
+		using namespace Catch;
+		if (bHasFailedOnce)
+		{
+			Xml.endElement(); //end failure element
+		}
+		bHasFailedOnce = false;
+		{
+			//using an element instead of an attribute on "testcase" to make streaming to the xml writer easier
+			//otherwise any failures would have to be stored as a member
+			XmlWriter::ScopedElement ResultElement = Xml.scopedElement("result");
+			ResultElement.writeAttribute("success", TestCaseStats.totals.assertions.allOk());
+		}
+		Xml.endElement();
+	}
+
+	virtual void skipTest(Catch::TestCaseInfo const& TestInfo) override
+	{
+		Catch::XmlWriter::ScopedElement TestCaseElement = Xml.scopedElement("testcase");
+		TestCaseElement.writeAttribute("name", TestInfo.name);
+		TestCaseElement.writeAttribute("skipped", "true");
+	}
+
+	static std::string getDescription()
+	{
+		return "Reporter for LowLevelTests";
+	}
+
+	void assertionEnded(Catch::AssertionStats const& Stats) override
+	{
+		using namespace Catch;
+
+		AssertionResult const& Result = Stats.assertionResult;
+
+		if (!Result.isOk())
+		{
+			const char* ElementName;
+			switch (Result.getResultType())
+			{
+			case ResultWas::ThrewException:
+			case ResultWas::FatalErrorCondition:
+			case ResultWas::ExplicitFailure:
+			case ResultWas::ExpressionFailed:
+			case ResultWas::DidntThrowException:
+				ElementName = "failure";
+				break;
+
+			default:
+				ElementName = "internalError";
+				break;
+			}
+
+			if (!bHasFailedOnce)
+			{
+				bHasFailedOnce = true;
+				Xml.startElement(ElementName);
+			}
+			TAnsiStringBuilder<1024> Stream;
+			if (Stats.totals.assertions.total() > 0)
+			{
+				Stream << "FAILED" << ":\n";
+				if (Result.hasExpression())
+				{
+					Stream << "  " 
+						<< Result.getExpressionInMacro().c_str()
+						<< "\n";
+				}
+				if (Result.hasExpandedExpression())
+				{
+					Stream << "with expansion:\n" 
+						<< Result.getExpandedExpression().c_str() << "\n";
+				}
+			}
+			else 
+			{
+				Stream << '\n';
+			}
+
+			if (!Result.getMessage().empty())
+			{
+				Stream << Result.getMessage().data() << '\n';
+			}
+			for (auto const& Msg : Stats.infoMessages)
+			{
+				if (Msg.type == ResultWas::Info)
+				{
+					Stream << Msg.message.data() << '\n';
+				}
+			}
+
+			Stream << "at " << Result.getSourceInfo().file << '(' << (uint64)Result.getSourceInfo().line << ')';
+			Xml.writeText(Stream.ToString(), XmlFormatting::Newline);
+		}
+	}
+
+
+private:
+	Catch::XmlWriter Xml;
+	bool bHasFailedOnce = false;
+};
+
+
+CATCH_REGISTER_REPORTER("unreal", FUnrealReporter)
+
+}
\ No newline at end of file
diff --git a/Engine/Source/Developer/LowLevelTestsRunner/Public/TestMacros/Assertions.h b/Engine/Source/Developer/LowLevelTestsRunner/Public/TestMacros/Assertions.h
--- a/Engine/Source/Developer/LowLevelTestsRunner/Public/TestMacros/Assertions.h
+++ b/Engine/Source/Developer/LowLevelTestsRunner/Public/TestMacros/Assertions.h
@@ -6,26 +6,26 @@
 #include "LowLevelTestsRunner/CheckScope.h"
 
 //requires that an UE `ensure` fails in this call
-#define REQUIRE_ENSURE(...) INTERNAL_UE_ENSURE( "REQUIRE_ENSURE", Catch::ResultDisposition::Normal, __VA_ARGS__ )
+#define REQUIRE_ENSURE(...) INTERNAL_UE_ENSURE( "REQUIRE_ENSURE", Catch::ResultDisposition::Normal, #__VA_ARGS__, __VA_ARGS__ )
 
 //requires that an UE `ensure` fails with a message that matches the supplied message
-#define REQUIRE_ENSURE_MSG(msg, ...) INTERNAL_UE_ENSURE_MSG(msg, "REQUIRE_ENSURE", Catch::ResultDisposition::Normal, __VA_ARGS__ )
+#define REQUIRE_ENSURE_MSG(msg, ...) INTERNAL_UE_ENSURE_MSG(msg, "REQUIRE_ENSURE", Catch::ResultDisposition::Normal, #__VA_ARGS__, __VA_ARGS__ )
 
 //checks that an UE `ensure` fails in this call
-#define CHECK_ENSURE(...) INTERNAL_UE_ENSURE( "CHECK_ENSURE", Catch::ResultDisposition::ContinueOnFailure, __VA_ARGS__ )
+#define CHECK_ENSURE(...) INTERNAL_UE_ENSURE( "CHECK_ENSURE", Catch::ResultDisposition::ContinueOnFailure, #__VA_ARGS__, __VA_ARGS__ )
 
 //checks that an UE `ensure` fails with a message that matches the supplied message
-#define CHECK_ENSURE_MSG(msg, ...) INTERNAL_UE_ENSURE_MSG(msg, "CHECK_ENSURE", Catch::ResultDisposition::ContinueOnFailure, __VA_ARGS__ )
+#define CHECK_ENSURE_MSG(msg, ...) INTERNAL_UE_ENSURE_MSG(msg, "CHECK_ENSURE", Catch::ResultDisposition::ContinueOnFailure, #__VA_ARGS__, __VA_ARGS__ )
 
 //requires that a UE `check` fails in this call
-#define REQUIRE_CHECK(...) INTERNAL_UE_CHECK( "REQUIRE_ENSURE", Catch::ResultDisposition::Normal, __VA_ARGS__ )
+#define REQUIRE_CHECK(...) INTERNAL_UE_CHECK( "REQUIRE_ENSURE", Catch::ResultDisposition::Normal, #__VA_ARGS__, __VA_ARGS__ )
 
 //requires that a UE `check` fails in this call containg the supplies message
-#define REQUIRE_CHECK_MSG(msg, ...) INTERNAL_UE_CHECK_MSG(msg, "REQUIRE_ENSURE", Catch::ResultDisposition::Normal, __VA_ARGS__ )
+#define REQUIRE_CHECK_MSG(msg, ...) INTERNAL_UE_CHECK_MSG(msg, "REQUIRE_ENSURE", Catch::ResultDisposition::Normal, #__VA_ARGS__, __VA_ARGS__ )
 
-#define INTERNAL_UE_ENSURE( macroName, resultDisposition, ... ) \
+#define INTERNAL_UE_ENSURE( macroName, resultDisposition, ensureExpr, ... ) \
 	do { \
-		Catch::AssertionHandler catchAssertionHandler( macroName##_catch_sr, CATCH_INTERNAL_LINEINFO, CATCH_INTERNAL_STRINGIFY(__VA_ARGS__), resultDisposition ); \
+		Catch::AssertionHandler catchAssertionHandler( macroName##_catch_sr, CATCH_INTERNAL_LINEINFO, ensureExpr, resultDisposition ); \
 		INTERNAL_CATCH_TRY { \
 			::UE::LowLevelTests::FEnsureScope scope; \
 			static_cast<void>(__VA_ARGS__); \
@@ -36,9 +36,9 @@
 		INTERNAL_CATCH_REACT( catchAssertionHandler ) \
 	} while(false) \
 
-#define INTERNAL_UE_ENSURE_MSG(msg, macroName, resultDisposition, ... ) \
+#define INTERNAL_UE_ENSURE_MSG(msg, macroName, resultDisposition, ensureExpr, ... ) \
 	do { \
-		Catch::AssertionHandler catchAssertionHandler( macroName##_catch_sr, CATCH_INTERNAL_LINEINFO, CATCH_INTERNAL_STRINGIFY(__VA_ARGS__), resultDisposition ); \
+		Catch::AssertionHandler catchAssertionHandler( macroName##_catch_sr, CATCH_INTERNAL_LINEINFO, ensureExpr, resultDisposition ); \
 		INTERNAL_CATCH_TRY { \
 			::UE::LowLevelTests::FEnsureScope scope(msg); \
 			static_cast<void>(__VA_ARGS__); \
@@ -49,9 +49,9 @@
 		INTERNAL_CATCH_REACT( catchAssertionHandler ) \
 	} while(false) \
 
-#define INTERNAL_UE_CHECK(macroName, resultDisposition, ... ) \
+#define INTERNAL_UE_CHECK(macroName, resultDisposition, checkExpr, ... ) \
 	do { \
-		Catch::AssertionHandler catchAssertionHandler( macroName##_catch_sr, CATCH_INTERNAL_LINEINFO, CATCH_INTERNAL_STRINGIFY(__VA_ARGS__), resultDisposition ); \
+		Catch::AssertionHandler catchAssertionHandler( macroName##_catch_sr, CATCH_INTERNAL_LINEINFO, checkExpr, resultDisposition ); \
 		INTERNAL_CATCH_TRY { \
 			::UE::LowLevelTests::FCheckScope scope; \
 			__VA_ARGS__; \
@@ -62,9 +62,9 @@
 		INTERNAL_CATCH_REACT( catchAssertionHandler ) \
 	} while(false) \
 
-#define INTERNAL_UE_CHECK_MSG(msg, macroName, resultDisposition, ... ) \
+#define INTERNAL_UE_CHECK_MSG(msg, macroName, resultDisposition, checkExpr, ... ) \
 	do { \
-		Catch::AssertionHandler catchAssertionHandler( macroName##_catch_sr, CATCH_INTERNAL_LINEINFO, CATCH_INTERNAL_STRINGIFY(__VA_ARGS__), resultDisposition ); \
+		Catch::AssertionHandler catchAssertionHandler( macroName##_catch_sr, CATCH_INTERNAL_LINEINFO, checkExpr, resultDisposition ); \
 		INTERNAL_CATCH_TRY { \
 			::UE::LowLevelTests::FCheckScope scope(msg); \
 			__VA_ARGS__; \
diff --git a/Engine/Source/Programs/UnrealBuildTool/Configuration/TestTargetRules.cs b/Engine/Source/Programs/UnrealBuildTool/Configuration/TestTargetRules.cs
--- a/Engine/Source/Programs/UnrealBuildTool/Configuration/TestTargetRules.cs
+++ b/Engine/Source/Programs/UnrealBuildTool/Configuration/TestTargetRules.cs
@@ -95,12 +95,24 @@ namespace UnrealBuildTool
 			bExplicitTestsTargetOverride = this.GetType() != typeof(TestTargetRules);
 
 			ExeBinariesSubFolder = LaunchModuleName = Name + (bExplicitTestsTargetOverride ? string.Empty : "Tests");
+			IncludeOrderVersion = EngineIncludeOrderVersion.Latest;
 
 			if (bExplicitTestsTargetOverride)
 			{
 				bBuildInSolutionByDefault = true;
 				SolutionDirectory = "Programs/LowLevelTests";
 			}
+
+			if (Target.Platform == UnrealTargetPlatform.Win64)
+			{
+				string OutputName = "$(TargetName)";
+				if(Target.Configuration != UndecoratedConfiguration)
+				{
+					OutputName = OutputName + "-" + Target.Platform + "-" + Target.Configuration;
+				}
+				string OutputFileName = $"$(EngineDir)\\Binaries\\Win64\\{ExeBinariesSubFolder}\\{OutputName}.exe.is_unreal_test";
+				PostBuildSteps.Add("echo > " + OutputFileName);
+			}
 		}
 
 		/// <summary>
