Compiler fixup from CL: 23064620

diff --git a/Engine/Source/Developer/LowLevelTestsRunner/Private/TestMain.cpp b/Engine/Source/Developer/LowLevelTestsRunner/Private/TestMain.cpp
--- a/Engine/Source/Developer/LowLevelTestsRunner/Private/TestMain.cpp
+++ b/Engine/Source/Developer/LowLevelTestsRunner/Private/TestMain.cpp
@@ -5,4 +5,3 @@
 TCHAR GInternalProjectName[64] = TEXT("LowLevelTests");
 const TCHAR* GForeignEngineDir = nullptr;
 FEngineLoop GEngineLoop;
-bool GIsConsoleExecutable = true;
