Online/ OSS tests: Catch register tests using delegates, setup delegate is invoked during TestRunner::GlobalSetup and teardown during TestRunner::GlobalTeardown

diff --git a/Engine/Source/Developer/LowLevelTestsRunner/Private/TestRunner.cpp b/Engine/Source/Developer/LowLevelTestsRunner/Private/TestRunner.cpp
--- a/Engine/Source/Developer/LowLevelTestsRunner/Private/TestRunner.cpp
+++ b/Engine/Source/Developer/LowLevelTestsRunner/Private/TestRunner.cpp
@@ -26,7 +26,6 @@
 
 namespace UE::LowLevelTests
 {
-
 static ITestRunner* GTestRunner;
 
 ITestRunner* ITestRunner::Get()
@@ -229,6 +228,11 @@ void FTestRunner::GlobalSetup()
 
 		FLogSuppressionInterface::Get().ProcessConfigAndCommandLine();
 	}
+
+	if (FTestDelegates::GlobalSetup.Get()->IsBound())
+	{
+		FTestDelegates::GlobalSetup.Get()->Execute();
+	}
 }
 
 void FTestRunner::GlobalTeardown() const
@@ -244,6 +248,11 @@ void FTestRunner::GlobalTeardown() const
 		GError = ErrorOutputDevice.GetDeviceError();
 	}
 
+	if (FTestDelegates::GlobalTeardown.Get()->IsBound())
+	{
+		FTestDelegates::GlobalTeardown.Get()->Execute();
+	}
+
 	CleanupPlatform();
 }
 
diff --git a/Engine/Source/Developer/LowLevelTestsRunner/Public/TestRunner.h b/Engine/Source/Developer/LowLevelTestsRunner/Public/TestRunner.h
--- a/Engine/Source/Developer/LowLevelTestsRunner/Public/TestRunner.h
+++ b/Engine/Source/Developer/LowLevelTestsRunner/Public/TestRunner.h
@@ -11,6 +11,7 @@
 #include "Logging/LogSuppressionInterface.h"
 #include "Modules/ModuleManager.h"
 #include "Misc/CommandLine.h"
+#include "Templates/SharedPointer.h"
 #include <cstring>
 #include <thread>
 #include <chrono>
@@ -27,4 +28,13 @@
 #endif
 const char* GetCacheDirectory();
 
+DECLARE_DELEGATE(FGlobalSetupHandler);
+
+class FTestDelegates
+{
+public:
+	inline static TSharedPtr<FGlobalSetupHandler, ESPMode::ThreadSafe> GlobalSetup = MakeShareable(new FGlobalSetupHandler());
+	inline static TSharedPtr<FGlobalSetupHandler, ESPMode::ThreadSafe> GlobalTeardown = MakeShareable(new FGlobalSetupHandler());
+};
+
 int RunTests(int argc, const char* argv[]);
\ No newline at end of file
