compile fix for when `DO_ENSURE` is off
and compile fix for SHIPPING

diff --git a/Engine/Source/Developer/LowLevelTestsRunner/Private/CheckScope.cpp b/Engine/Source/Developer/LowLevelTestsRunner/Private/CheckScope.cpp
--- a/Engine/Source/Developer/LowLevelTestsRunner/Private/CheckScope.cpp
+++ b/Engine/Source/Developer/LowLevelTestsRunner/Private/CheckScope.cpp
@@ -63,9 +63,15 @@ private:
 
 FCheckScope::FCheckScope(const ANSICHAR* Msg)
 	: DeviceError(new FCheckScopeOutputDeviceError(GError, Msg))
+#if !UE_BUILD_SHIPPING
 	, bIgnoreDebugger(GIgnoreDebugger)
+#else
+	, bIgnoreDebugger(false)
+#endif
 {
+#if !UE_BUILD_SHIPPING
 	GIgnoreDebugger = true;
+#endif
 	GError = DeviceError;
 }
 
@@ -76,7 +82,9 @@ FCheckScope::FCheckScope()
 
 FCheckScope::~FCheckScope()
 {
+#if !UE_BUILD_SHIPPING
 	GIgnoreDebugger = bIgnoreDebugger;
+#endif
 	GError = DeviceError->GetDeviceError();
 	delete DeviceError;
 }
diff --git a/Engine/Source/Developer/LowLevelTestsRunner/Private/EnsureScope.cpp b/Engine/Source/Developer/LowLevelTestsRunner/Private/EnsureScope.cpp
--- a/Engine/Source/Developer/LowLevelTestsRunner/Private/EnsureScope.cpp
+++ b/Engine/Source/Developer/LowLevelTestsRunner/Private/EnsureScope.cpp
@@ -23,6 +23,7 @@ FEnsureScope::FEnsureScope(const ANSICHAR* ExpectedMsg)
 FEnsureScope::FEnsureScope(TFunction<bool(const FEnsureHandlerArgs& Args)> EnsureFunc)
 	: Count(0)
 {
+#if DO_ENSURE
 	OldHandler = SetEnsureHandler([this, EnsureFunc](const FEnsureHandlerArgs& Args) -> bool
 		{
 			bool Handled = EnsureFunc(Args);
@@ -32,11 +33,14 @@ FEnsureScope::FEnsureScope(TFunction<bool(const FEnsureHandlerArgs& Args)> Ensur
 			}
 			return Handled;
 		});
+#endif
 }
 
 FEnsureScope::~FEnsureScope()
 {
+#if DO_ENSURE
 	SetEnsureHandler(OldHandler);
+#endif
 }
 
 }
\ No newline at end of file
