LowLevelTests: Removed unnecessary public includes of DDC

#rnx

diff --git a/Engine/Source/Developer/LowLevelTestsRunner/Private/TestCommon/EditorUtilities.cpp b/Engine/Source/Developer/LowLevelTestsRunner/Private/TestCommon/EditorUtilities.cpp
--- a/Engine/Source/Developer/LowLevelTestsRunner/Private/TestCommon/EditorUtilities.cpp
+++ b/Engine/Source/Developer/LowLevelTestsRunner/Private/TestCommon/EditorUtilities.cpp
@@ -2,7 +2,16 @@
 
 #if WITH_EDITOR
 
-#include "TestCommon/EditorUtilities.h"
+#include "Misc/QueuedThreadPool.h"
+
+#if WITH_COREUOBJECT
+#include "UObject/PackageResourceManager.h"
+#endif
+
+#ifdef DERIVEDDATACACHE_API
+#include "DerivedDataBuild.h"
+#include "DerivedDataCache.h"
+#endif
 
 void InitEditorThreadPools()
 {
@@ -35,8 +44,8 @@ void InitDerivedDataCache()
 	if (!FPlatformProperties::RequiresCookedData())
 	{
 		// Ensure that DDC is initialized from the game thread.
-		UE::DerivedData::GetBuild();
 		UE::DerivedData::GetCache();
+		UE::DerivedData::GetBuild();
 		GetDerivedDataCacheRef();
 	}
 #endif // WITH_EDITORONLY_DATA && defined(DERIVEDDATACACHE_API)
@@ -80,4 +89,4 @@ void InitSlate()
 	CurrentSlateApp.InitializeRenderer(SlateRendererSharedRef);
 }
 
-#endif
+#endif // WITH_EDITOR
diff --git a/Engine/Source/Developer/LowLevelTestsRunner/Private/TestCommon/Initialization.cpp b/Engine/Source/Developer/LowLevelTestsRunner/Private/TestCommon/Initialization.cpp
--- a/Engine/Source/Developer/LowLevelTestsRunner/Private/TestCommon/Initialization.cpp
+++ b/Engine/Source/Developer/LowLevelTestsRunner/Private/TestCommon/Initialization.cpp
@@ -3,7 +3,6 @@
 #include "TestCommon/Initialization.h"
 #include "TestCommon/CoreUtilities.h"
 #include "TestCommon/CoreUObjectUtilities.h"
-#include "TestCommon/EditorUtilities.h"
 #include "TestCommon/EngineUtilities.h"
 
 
diff --git a/Engine/Source/Developer/LowLevelTestsRunner/Public/TestCommon/EditorUtilities.h b/Engine/Source/Developer/LowLevelTestsRunner/Public/TestCommon/EditorUtilities.h
deleted file mode 100644
