fix for FName not resolving to strings in the debugger
changed testmain in lowleveltestrunner to use IMPLEMENT_APPLICATION which does all the correct things
https://p4-swarm.epicgames.net/reviews/23252374

diff --git a/Engine/Source/Developer/LowLevelTestsRunner/Private/TestMain.cpp b/Engine/Source/Developer/LowLevelTestsRunner/Private/TestMain.cpp
--- a/Engine/Source/Developer/LowLevelTestsRunner/Private/TestMain.cpp
+++ b/Engine/Source/Developer/LowLevelTestsRunner/Private/TestMain.cpp
@@ -2,6 +2,13 @@
 
 #include "RequiredProgramMainCPPInclude.h"
 
-TCHAR GInternalProjectName[64] = TEXT("LowLevelTests");
-const TCHAR* GForeignEngineDir = nullptr;
-FEngineLoop GEngineLoop;
+//typical defined by TargetRules but LowLevelTestRunner is not setup correctly
+//should revist this in the future
+#ifndef IMPLEMENT_ENCRYPTION_KEY_REGISTRATION
+	#define IMPLEMENT_ENCRYPTION_KEY_REGISTRATION()
+#endif
+#ifndef IMPLEMENT_SIGNING_KEY_REGISTRATION
+	#define IMPLEMENT_SIGNING_KEY_REGISTRATION()
+#endif
+
+IMPLEMENT_APPLICATION(LowLevelTests, "LowLevelTests");
