changed REQUIRE_CHECK and REQUIRE_ENSURE test macros to disable themselves if DO_CHECK or DO_ENSURE it not enabled
this allows usage of REQUIRE_CHECK to not have to be guarded

https://p4-swarm.epicgames.net/reviews/23784409

diff --git a/Engine/Source/Developer/LowLevelTestsRunner/Public/TestMacros/Assertions.h b/Engine/Source/Developer/LowLevelTestsRunner/Public/TestMacros/Assertions.h
--- a/Engine/Source/Developer/LowLevelTestsRunner/Public/TestMacros/Assertions.h
+++ b/Engine/Source/Developer/LowLevelTestsRunner/Public/TestMacros/Assertions.h
@@ -6,70 +6,70 @@
 #include "LowLevelTestsRunner/CheckScope.h"
 
 //requires that an UE `ensure` fails in this call
-#define REQUIRE_ENSURE(...) INTERNAL_UE_ENSURE( "REQUIRE_ENSURE", Catch::ResultDisposition::Normal, #__VA_ARGS__, __VA_ARGS__ )
+#define REQUIRE_ENSURE(...) INTERNAL_UE_ENSURE( "REQUIRE_ENSURE", DO_ENSURE, Catch::ResultDisposition::Normal, #__VA_ARGS__, __VA_ARGS__ )
 
 //requires that an UE `ensure` fails with a message that matches the supplied message
-#define REQUIRE_ENSURE_MSG(msg, ...) INTERNAL_UE_ENSURE_MSG(msg, "REQUIRE_ENSURE", Catch::ResultDisposition::Normal, #__VA_ARGS__, __VA_ARGS__ )
+#define REQUIRE_ENSURE_MSG(msg, ...) INTERNAL_UE_ENSURE_MSG(msg, "REQUIRE_ENSURE", DO_ENSURE, Catch::ResultDisposition::Normal, #__VA_ARGS__, __VA_ARGS__ )
 
 //checks that an UE `ensure` fails in this call
-#define CHECK_ENSURE(...) INTERNAL_UE_ENSURE( "CHECK_ENSURE", Catch::ResultDisposition::ContinueOnFailure, #__VA_ARGS__, __VA_ARGS__ )
+#define CHECK_ENSURE(...) INTERNAL_UE_ENSURE( "CHECK_ENSURE", DO_ENSURE, Catch::ResultDisposition::ContinueOnFailure, #__VA_ARGS__, __VA_ARGS__ )
 
 //checks that an UE `ensure` fails with a message that matches the supplied message
-#define CHECK_ENSURE_MSG(msg, ...) INTERNAL_UE_ENSURE_MSG(msg, "CHECK_ENSURE", Catch::ResultDisposition::ContinueOnFailure, #__VA_ARGS__, __VA_ARGS__ )
+#define CHECK_ENSURE_MSG(msg, ...) INTERNAL_UE_ENSURE_MSG(msg, "CHECK_ENSURE", DO_ENSURE, Catch::ResultDisposition::ContinueOnFailure, #__VA_ARGS__, __VA_ARGS__ )
 
 //requires that a UE `check` fails in this call
-#define REQUIRE_CHECK(...) INTERNAL_UE_CHECK( "REQUIRE_ENSURE", Catch::ResultDisposition::Normal, #__VA_ARGS__, __VA_ARGS__ )
+#define REQUIRE_CHECK(...) INTERNAL_UE_CHECK( "REQUIRE_ENSURE", DO_CHECK, Catch::ResultDisposition::Normal, #__VA_ARGS__, __VA_ARGS__ )
 
 //requires that a UE `check` fails in this call containg the supplies message
-#define REQUIRE_CHECK_MSG(msg, ...) INTERNAL_UE_CHECK_MSG(msg, "REQUIRE_ENSURE", Catch::ResultDisposition::Normal, #__VA_ARGS__, __VA_ARGS__ )
+#define REQUIRE_CHECK_MSG(msg, ...) INTERNAL_UE_CHECK_MSG(msg, "REQUIRE_ENSURE", DO_CHECK, Catch::ResultDisposition::Normal, #__VA_ARGS__, __VA_ARGS__ )
 
-#define INTERNAL_UE_ENSURE( macroName, resultDisposition, ensureExpr, ... ) \
+#define INTERNAL_UE_ENSURE( macroName, doEnsure, resultDisposition, ensureExpr, ... ) \
 	do { \
 		Catch::AssertionHandler catchAssertionHandler( macroName##_catch_sr, CATCH_INTERNAL_LINEINFO, ensureExpr, resultDisposition ); \
 		INTERNAL_CATCH_TRY { \
 			::UE::LowLevelTests::FEnsureScope scope; \
 			static_cast<void>(__VA_ARGS__); \
 			bool bEncounteredEnsure = scope.GetCount() > 0; \
-			if (!bEncounteredEnsure) \
+			if (doEnsure && !bEncounteredEnsure) \
 				catchAssertionHandler.handleMessage( Catch::ResultWas::ExplicitFailure ,"Expected failure of `ensure` not recieved."); \
 		} INTERNAL_CATCH_CATCH( catchAssertionHandler ) \
 		INTERNAL_CATCH_REACT( catchAssertionHandler ) \
 	} while(false) \
 
-#define INTERNAL_UE_ENSURE_MSG(msg, macroName, resultDisposition, ensureExpr, ... ) \
+#define INTERNAL_UE_ENSURE_MSG(msg, macroName, doEnsure, resultDisposition, ensureExpr, ... ) \
 	do { \
 		Catch::AssertionHandler catchAssertionHandler( macroName##_catch_sr, CATCH_INTERNAL_LINEINFO, ensureExpr, resultDisposition ); \
 		INTERNAL_CATCH_TRY { \
 			::UE::LowLevelTests::FEnsureScope scope(msg); \
 			static_cast<void>(__VA_ARGS__); \
 			bool bEncounteredEnsure = scope.GetCount() > 0; \
-			if (!bEncounteredEnsure) \
+			if (doEnsure && !bEncounteredEnsure) \
 				catchAssertionHandler.handleMessage( Catch::ResultWas::ExplicitFailure, ( Catch::MessageStream() << "Expected failure of `ensure` with message '" << msg << "' not recieved" + ::Catch::StreamEndStop() ).m_stream.str() ); \
 		} INTERNAL_CATCH_CATCH( catchAssertionHandler ) \
 		INTERNAL_CATCH_REACT( catchAssertionHandler ) \
 	} while(false) \
 
-#define INTERNAL_UE_CHECK(macroName, resultDisposition, checkExpr, ... ) \
+#define INTERNAL_UE_CHECK(macroName, doCheck, resultDisposition, checkExpr, ... ) \
 	do { \
 		Catch::AssertionHandler catchAssertionHandler( macroName##_catch_sr, CATCH_INTERNAL_LINEINFO, checkExpr, resultDisposition ); \
 		INTERNAL_CATCH_TRY { \
 			::UE::LowLevelTests::FCheckScope scope; \
 			__VA_ARGS__; \
 			bool bEncounteredEnsure = scope.GetCount() > 0; \
-			if (!bEncounteredEnsure) \
+			if (doCheck && !bEncounteredEnsure) \
 				catchAssertionHandler.handleMessage( Catch::ResultWas::ExplicitFailure, "Expected failure of `check` not recieved"); \
 		} INTERNAL_CATCH_CATCH( catchAssertionHandler ) \
 		INTERNAL_CATCH_REACT( catchAssertionHandler ) \
 	} while(false) \
 
-#define INTERNAL_UE_CHECK_MSG(msg, macroName, resultDisposition, checkExpr, ... ) \
+#define INTERNAL_UE_CHECK_MSG(msg, macroName, doCheck, resultDisposition, checkExpr, ... ) \
 	do { \
 		Catch::AssertionHandler catchAssertionHandler( macroName##_catch_sr, CATCH_INTERNAL_LINEINFO, checkExpr, resultDisposition ); \
 		INTERNAL_CATCH_TRY { \
 			::UE::LowLevelTests::FCheckScope scope(msg); \
 			__VA_ARGS__; \
 			bool bEncounteredEnsure = scope.GetCount() > 0; \
-			if (!bEncounteredEnsure) \
+			if (doCheck && !bEncounteredEnsure) \
 				catchAssertionHandler.handleMessage( Catch::ResultWas::ExplicitFailure, ( Catch::MessageStream() << "Expected failure of `check` containing message '" << msg << "' not recieved" + ::Catch::StreamEndStop() ).m_stream.str() ); \
 		} INTERNAL_CATCH_CATCH( catchAssertionHandler ) \
 		INTERNAL_CATCH_REACT( catchAssertionHandler ) \
diff --git a/Engine/Source/Runtime/Core/Tests/Misc/AssertionMacrosTest.cpp b/Engine/Source/Runtime/Core/Tests/Misc/AssertionMacrosTest.cpp
--- a/Engine/Source/Runtime/Core/Tests/Misc/AssertionMacrosTest.cpp
+++ b/Engine/Source/Runtime/Core/Tests/Misc/AssertionMacrosTest.cpp
@@ -2,7 +2,7 @@
 
 #include "Misc/Build.h" 
 
-#if WITH_LOW_LEVEL_TESTS && DO_ENSURE
+#if WITH_LOW_LEVEL_TESTS
 
 #include "Misc/AssertionMacros.h"
 #include "GenericPlatform/GenericPlatformMisc.h"
