added a WarnFilterScope for filtering Warning messages from appear in the log for tests
Horde greps the log for warnings and will flag the warnings if they appear

https://p4-swarm.epicgames.net/reviews/23816461

diff --git a/Engine/Source/Developer/LowLevelTestsRunner/Private/WarnFilterScope.cpp b/Engine/Source/Developer/LowLevelTestsRunner/Private/WarnFilterScope.cpp
new file mode 100644
--- /dev/null
+++ b/Engine/Source/Developer/LowLevelTestsRunner/Private/WarnFilterScope.cpp
@@ -0,0 +1,49 @@
+// Copyright Epic Games, Inc. All Rights Reserved.
+
+#include "LowLevelTestsRunner/WarnFilterScope.h"
+#include "Misc/FeedbackContext.h"
+
+namespace UE::Testing
+{
+	struct FFilterFeedback : public FFeedbackContext
+	{
+
+		FFilterFeedback(FFeedbackContext* OldFeedback, TFunction<bool(const TCHAR* Message, ELogVerbosity::Type Verbosity, const FName& Category)> Handler)
+			: OldFeedback(OldFeedback)
+			, Handler(Handler)
+		{
+
+		}
+
+		virtual void Serialize(const TCHAR* V, ELogVerbosity::Type Verbosity, const FName& Category) override
+		{
+			if (!Handler(V, Verbosity, Category))
+			{
+				OldFeedback->Serialize(V, Verbosity, Category);
+			}
+		}
+		virtual void Serialize(const TCHAR* V, ELogVerbosity::Type Verbosity, const FName& Category, double Time) override
+		{
+			if (!Handler(V, Verbosity, Category))
+			{
+				OldFeedback->Serialize(V, Verbosity, Category);
+			}
+		}
+
+		FFeedbackContext* OldFeedback;
+		TFunction<bool(const TCHAR* Message, ELogVerbosity::Type Verbosity, const FName& Category)> Handler;
+	};
+
+	FWarnFilterScope::FWarnFilterScope(TFunction<bool(const TCHAR* Message, ELogVerbosity::Type Verbosity, const FName& Category)> LogHandler)
+		: OldWarn(GWarn)
+		, Feedback(new FFilterFeedback(OldWarn, LogHandler))
+	{
+		GWarn = Feedback;
+	}
+
+	FWarnFilterScope::~FWarnFilterScope()
+	{
+		GWarn = OldWarn;
+		delete Feedback;
+	}
+}
\ No newline at end of file
diff --git a/Engine/Source/Developer/LowLevelTestsRunner/Public/LowLevelTestsRunner/WarnFilterScope.h b/Engine/Source/Developer/LowLevelTestsRunner/Public/LowLevelTestsRunner/WarnFilterScope.h
new file mode 100644
--- /dev/null
+++ b/Engine/Source/Developer/LowLevelTestsRunner/Public/LowLevelTestsRunner/WarnFilterScope.h
@@ -0,0 +1,26 @@
+// Copyright Epic Games, Inc. All Rights Reserved.
+
+#pragma once
+#include "HAL/Platform.h"
+#include "Templates/Function.h"
+
+class FFeedbackContext;
+
+namespace UE::Testing
+{
+	struct FFilterFeedback;
+
+	///@brief Scope that captures warning log messages
+	///necessary to filter out warning message from appearing the log file as Horde will flag the warnings
+	struct FWarnFilterScope
+	{
+		/// @param LogHandler function to filter log messages. returning true filters message from log file
+		FWarnFilterScope(TFunction<bool(const TCHAR* Message, ELogVerbosity::Type Verbosity, const FName& Category)> LogHandler);
+
+		~FWarnFilterScope();
+	private:
+		FFeedbackContext* OldWarn;
+		FFilterFeedback* Feedback;
+	};
+
+}
\ No newline at end of file
diff --git a/Engine/Source/Runtime/Core/Tests/Misc/WarnFilterScopeTest.cpp b/Engine/Source/Runtime/Core/Tests/Misc/WarnFilterScopeTest.cpp
new file mode 100644
--- /dev/null
+++ b/Engine/Source/Runtime/Core/Tests/Misc/WarnFilterScopeTest.cpp
@@ -0,0 +1,24 @@
+// Copyright Epic Games, Inc. All Rights Reserved.
+
+#if WITH_LOW_LEVEL_TESTS
+
+#include "TestHarness.h"
+#include "LowLevelTestsRunner/WarnFilterScope.h"
+#include "Logging/LogMacros.h"
+
+DECLARE_LOG_CATEGORY_EXTERN(LogTesting, Log, All);
+DEFINE_LOG_CATEGORY(LogTesting);
+
+TEST_CASE("UE::Testing::WarnFilterScope")
+{
+	bool bFiltered = false;
+	UE::Testing::FWarnFilterScope _([&bFiltered](const TCHAR* Message, ELogVerbosity::Type Verbosity, const FName& Category)
+		{
+			bFiltered = FCString::Strcmp(Message, TEXT("Test Message")) == 0 && Verbosity == ELogVerbosity::Type::Warning && Category == TEXT("LogTesting");
+			return bFiltered;
+		});
+	UE_LOG(LogTesting, Warning, TEXT("Test Message"));
+	CHECK(bFiltered);
+}
+
+#endif
