some minor fixes to WarnFilterScope
remove assignment and copy constructors

https://p4-swarm.epicgames.net/reviews/23876768

diff --git a/Engine/Source/Developer/LowLevelTestsRunner/Private/WarnFilterScope.cpp b/Engine/Source/Developer/LowLevelTestsRunner/Private/WarnFilterScope.cpp
--- a/Engine/Source/Developer/LowLevelTestsRunner/Private/WarnFilterScope.cpp
+++ b/Engine/Source/Developer/LowLevelTestsRunner/Private/WarnFilterScope.cpp
@@ -26,7 +26,7 @@ namespace UE::Testing
 		{
 			if (!Handler(V, Verbosity, Category))
 			{
-				OldFeedback->Serialize(V, Verbosity, Category);
+				OldFeedback->Serialize(V, Verbosity, Category, Time);
 			}
 		}
 
diff --git a/Engine/Source/Developer/LowLevelTestsRunner/Public/LowLevelTestsRunner/WarnFilterScope.h b/Engine/Source/Developer/LowLevelTestsRunner/Public/LowLevelTestsRunner/WarnFilterScope.h
--- a/Engine/Source/Developer/LowLevelTestsRunner/Public/LowLevelTestsRunner/WarnFilterScope.h
+++ b/Engine/Source/Developer/LowLevelTestsRunner/Public/LowLevelTestsRunner/WarnFilterScope.h
@@ -16,6 +16,8 @@ namespace UE::Testing
 	{
 		/// @param LogHandler function to filter log messages. returning true filters message from log file
 		FWarnFilterScope(TFunction<bool(const TCHAR* Message, ELogVerbosity::Type Verbosity, const FName& Category)> LogHandler);
+		FWarnFilterScope(const FWarnFilterScope&) = delete;
+		FWarnFilterScope& operator=(const FWarnFilterScope&) = delete;
 
 		~FWarnFilterScope();
 	private:
