Online LLT fix global setup

diff --git a/Engine/Source/Developer/LowLevelTestsRunner/Private/TestRunner.cpp b/Engine/Source/Developer/LowLevelTestsRunner/Private/TestRunner.cpp
--- a/Engine/Source/Developer/LowLevelTestsRunner/Private/TestRunner.cpp
+++ b/Engine/Source/Developer/LowLevelTestsRunner/Private/TestRunner.cpp
@@ -229,10 +229,7 @@ void FTestRunner::GlobalSetup()
 		FLogSuppressionInterface::Get().ProcessConfigAndCommandLine();
 	}
 
-	if (FTestDelegates::GlobalSetup.Get()->IsBound())
-	{
-		FTestDelegates::GlobalSetup.Get()->Execute();
-	}
+	FTestDelegates::GetGlobalSetup().ExecuteIfBound();
 }
 
 void FTestRunner::GlobalTeardown() const
@@ -248,10 +245,7 @@ void FTestRunner::GlobalTeardown() const
 		GError = ErrorOutputDevice.GetDeviceError();
 	}
 
-	if (FTestDelegates::GlobalTeardown.Get()->IsBound())
-	{
-		FTestDelegates::GlobalTeardown.Get()->Execute();
-	}
+	FTestDelegates::GetGlobalTeardown().ExecuteIfBound();
 
 	CleanupPlatform();
 }
@@ -307,4 +301,4 @@ int RunTests(int32 ArgC, const ANSICHAR* ArgV[])
 
 	int CatchReturn = TestRunner.RunCatchSession();
 	return CatchReturn;
-}
+}
\ No newline at end of file
diff --git a/Engine/Source/Developer/LowLevelTestsRunner/Public/TestRunner.h b/Engine/Source/Developer/LowLevelTestsRunner/Public/TestRunner.h
--- a/Engine/Source/Developer/LowLevelTestsRunner/Public/TestRunner.h
+++ b/Engine/Source/Developer/LowLevelTestsRunner/Public/TestRunner.h
@@ -33,8 +33,17 @@ DECLARE_DELEGATE(FGlobalSetupHandler);
 class FTestDelegates
 {
 public:
-	inline static TSharedPtr<FGlobalSetupHandler, ESPMode::ThreadSafe> GlobalSetup = MakeShareable(new FGlobalSetupHandler());
-	inline static TSharedPtr<FGlobalSetupHandler, ESPMode::ThreadSafe> GlobalTeardown = MakeShareable(new FGlobalSetupHandler());
+	static FGlobalSetupHandler& GetGlobalSetup()
+	{
+		static TUniquePtr<FGlobalSetupHandler> GlobalSetup = MakeUnique<FGlobalSetupHandler>();
+		return *GlobalSetup.Get();
+	}
+
+	static FGlobalSetupHandler& GetGlobalTeardown()
+	{
+		static TUniquePtr<FGlobalSetupHandler> GlobalTeardown = MakeUnique<FGlobalSetupHandler>();
+		return *GlobalTeardown.Get();
+	}
 };
 
 int RunTests(int argc, const char* argv[]);
\ No newline at end of file
