Add a couple calls to RunAndClearDelayedAutoRegisterDelegates from the low-level tests runner to run additional initialization code, required to run tests that depend on the engine module on some platforms.

diff --git a/Engine/Source/Developer/LowLevelTestsRunner/Private/TestCommon/CoreUtilities.cpp b/Engine/Source/Developer/LowLevelTestsRunner/Private/TestCommon/CoreUtilities.cpp
--- a/Engine/Source/Developer/LowLevelTestsRunner/Private/TestCommon/CoreUtilities.cpp
+++ b/Engine/Source/Developer/LowLevelTestsRunner/Private/TestCommon/CoreUtilities.cpp
@@ -6,6 +6,7 @@
 #include "HAL/PlatformProcess.h"
 #include "Logging/LogSuppressionInterface.h"
 #include "Misc/CommandLine.h"
+#include "Misc/DelayedAutoRegister.h"
 #include "Misc/QueuedThreadPool.h"
 
 #include "TestCommon/CoreUtilities.h"
@@ -98,6 +99,8 @@ void InitTaskGraph(bool MultiThreaded, ENamedThreads::Type ThreadToAttach)
 	FTaskTagScope Scope(ETaskTag::EGameThread);
 	FTaskGraphInterface::Startup(MultiThreaded ? FPlatformMisc::NumberOfWorkerThreadsToSpawn() : 1);
 	FTaskGraphInterface::Get().AttachToThread(ThreadToAttach);
+
+	FDelayedAutoRegisterHelper::RunAndClearDelayedAutoRegisterDelegates(EDelayedRegisterRunPhase::TaskGraphSystemReady);
 }
 
 void CleanupTaskGraph()
diff --git a/Engine/Source/Developer/LowLevelTestsRunner/Private/TestCommon/Initialization.cpp b/Engine/Source/Developer/LowLevelTestsRunner/Private/TestCommon/Initialization.cpp
--- a/Engine/Source/Developer/LowLevelTestsRunner/Private/TestCommon/Initialization.cpp
+++ b/Engine/Source/Developer/LowLevelTestsRunner/Private/TestCommon/Initialization.cpp
@@ -4,7 +4,7 @@
 #include "TestCommon/CoreUtilities.h"
 #include "TestCommon/CoreUObjectUtilities.h"
 #include "TestCommon/EngineUtilities.h"
-
+#include "Misc/DelayedAutoRegister.h"
 
 void InitAllThreadPoolsEditorEx(bool MultiThreaded)
 {
@@ -19,6 +19,8 @@ void InitStats()
 #if STATS
 	FThreadStats::StartThread();
 #endif // #if STATS
+
+	FDelayedAutoRegisterHelper::RunAndClearDelayedAutoRegisterDelegates(EDelayedRegisterRunPhase::StatSystemReady);
 }
 
 void InitAll(bool bAllowLogging, bool bMultithreaded)
