Add -buildmachine to LowLevelTests, so that DDC performance warnings can be downgraded to display.
#rnx

diff --git a/Engine/Source/Developer/LowLevelTestsRunner/Private/TestRunner.cpp b/Engine/Source/Developer/LowLevelTestsRunner/Private/TestRunner.cpp
--- a/Engine/Source/Developer/LowLevelTestsRunner/Private/TestRunner.cpp
+++ b/Engine/Source/Developer/LowLevelTestsRunner/Private/TestRunner.cpp
@@ -153,6 +153,10 @@ void FTestRunner::ParseCommandLine(TConstArrayView<const ANSICHAR*> Args)
 		{
 			bAttachToDebugger = true;
 		}
+		else if (Arg == ANSITEXTVIEW("--buildmachine"))
+		{
+			GIsBuildMachine = true;
+		}
 		else
 		{
 			CatchArgs.Add(Arg.GetData());
