LowLevelTests: Allow GInternalProjectName to be defined by UE_PROJECT_NAME

diff --git a/Engine/Source/Developer/LowLevelTestsRunner/Private/TestExterns.cpp b/Engine/Source/Developer/LowLevelTestsRunner/Private/TestExterns.cpp
--- a/Engine/Source/Developer/LowLevelTestsRunner/Private/TestExterns.cpp
+++ b/Engine/Source/Developer/LowLevelTestsRunner/Private/TestExterns.cpp
@@ -4,8 +4,13 @@
 #include "CoreTypes.h"
 #include "Modules/ModuleManager.h"
 
+// Allow UBT to name my application
+#if defined(UE_PROJECT_NAME)
 TCHAR GInternalProjectName[64] = TEXT("LowLevelTests");
-const TCHAR* GForeignEngineDir = nullptr;
+#else
+TCHAR GInternalProjectName[64] = TEXT( PREPROCESSOR_TO_STRING(UE_PROJECT_NAME) );
+#endif 
+IMPLEMENT_FOREIGN_ENGINE_DIR()
 
 #if !EXPLICIT_TESTS_TARGET
 bool GIsGameAgnosticExe = false;
@@ -20,5 +25,9 @@ bool GIsGameAgnosticExe = false;
 	#define IMPLEMENT_SIGNING_KEY_REGISTRATION()
 #endif
 
+#if defined(UE_TARGET_NAME)
+IMPLEMENT_TARGET_NAME_REGISTRATION()
+#endif
+
 // Debug visualizers and new operator overloads
 PER_MODULE_BOILERPLATE
\ No newline at end of file
