Added UE_KINDA_SMALL_NUMBER to Tolerance in Comparisons.h

diff --git a/Engine/Source/Developer/LowLevelTestsRunner/Public/TestCommon/Comparisons.h b/Engine/Source/Developer/LowLevelTestsRunner/Public/TestCommon/Comparisons.h
--- a/Engine/Source/Developer/LowLevelTestsRunner/Public/TestCommon/Comparisons.h
+++ b/Engine/Source/Developer/LowLevelTestsRunner/Public/TestCommon/Comparisons.h
@@ -49,7 +49,7 @@ namespace {
 	}
 #endif
 
-	bool TestEqual(const TCHAR* What, const float Actual, const float Expected, float Tolerance)
+	bool TestEqual(const TCHAR* What, const float Actual, const float Expected, float Tolerance = UE_KINDA_SMALL_NUMBER)
 	{
 		if (!FMath::IsNearlyEqual(Actual, Expected, Tolerance))
 		{
@@ -59,7 +59,7 @@ namespace {
 		return true;
 	}
 
-	bool TestEqual(const TCHAR* What, const double Actual, const double Expected, double Tolerance)
+	bool TestEqual(const TCHAR* What, const double Actual, const double Expected, double Tolerance = UE_KINDA_SMALL_NUMBER)
 	{
 		if (!FMath::IsNearlyEqual(Actual, Expected, Tolerance))
 		{
