Fixed the UE_PROJECT_NAME logic, that was  reversed

diff --git a/Engine/Source/Developer/LowLevelTestsRunner/Private/TestExterns.cpp b/Engine/Source/Developer/LowLevelTestsRunner/Private/TestExterns.cpp
--- a/Engine/Source/Developer/LowLevelTestsRunner/Private/TestExterns.cpp
+++ b/Engine/Source/Developer/LowLevelTestsRunner/Private/TestExterns.cpp
@@ -5,7 +5,7 @@
 #include "Modules/ModuleManager.h"
 
 // Allow UBT to name my application
-#if defined(UE_PROJECT_NAME)
+#if !defined(UE_PROJECT_NAME)
 TCHAR GInternalProjectName[64] = TEXT("LowLevelTests");
 #else
 TCHAR GInternalProjectName[64] = TEXT( PREPROCESSOR_TO_STRING(UE_PROJECT_NAME) );
