always initialize GWarn as its used by FScopedSlowTask
https://p4-swarm.epicgames.net/reviews/25633852

diff --git a/Engine/Source/Developer/LowLevelTestsRunner/Private/TestRunner.cpp b/Engine/Source/Developer/LowLevelTestsRunner/Private/TestRunner.cpp
--- a/Engine/Source/Developer/LowLevelTestsRunner/Private/TestRunner.cpp
+++ b/Engine/Source/Developer/LowLevelTestsRunner/Private/TestRunner.cpp
@@ -225,15 +225,14 @@ void FTestRunner::GlobalSetup()
 			Catch::getResultCapture().handleMessage(info, Catch::ResultWas::ExplicitFailure, StringCast<ANSICHAR>(*Error).Get(), reaction);
 		});
 
-	if (bLogOutput || bDebugMode)
-	{
-		// Set up GWarn to handle Error, Warning, Display; but only when log output is enabled.
+	// Set up GWarn to handle Error, Warning, Display; but only when log output is enabled.
 #if WITH_APPLICATION_CORE
-		GWarn = FPlatformApplicationMisc::GetFeedbackContext();
+	GWarn = FPlatformApplicationMisc::GetFeedbackContext();
 #else
-		GWarn = FPlatformOutputDevices::GetFeedbackContext();
+	GWarn = FPlatformOutputDevices::GetFeedbackContext();
 #endif
-
+	if (bLogOutput || bDebugMode)
+	{
 		// Set up default output devices to handle Log, Verbose, VeryVerbose.
 		FPlatformOutputDevices::SetupOutputDevices();
 
