finish initializing the garbage collector allowing tests to run GC
https://p4-swarm.epicgames.net/reviews/25665144

diff --git a/Engine/Source/Developer/LowLevelTestsRunner/Private/TestCommon/CoreUObjectUtilities.cpp b/Engine/Source/Developer/LowLevelTestsRunner/Private/TestCommon/CoreUObjectUtilities.cpp
--- a/Engine/Source/Developer/LowLevelTestsRunner/Private/TestCommon/CoreUObjectUtilities.cpp
+++ b/Engine/Source/Developer/LowLevelTestsRunner/Private/TestCommon/CoreUObjectUtilities.cpp
@@ -10,6 +10,7 @@
 #include "Modules/ModuleManager.h"
 #include "UObject/Package.h"
 #include "UObject/PackageResourceManager.h"
+#include "UObject/GCObject.h"
 
 #if WITH_ENGINE
 #if UE_LLT_USE_PLATFORM_FILE_STUB
@@ -93,6 +94,11 @@ void InitCoreUObject()
 
 		FDelayedAutoRegisterHelper::RunAndClearDelayedAutoRegisterDelegates(EDelayedRegisterRunPhase::ObjectSystemReady);
 	}
+	FGCObject::StaticInit();
+	if (GUObjectArray.IsOpenForDisregardForGC())
+	{
+		GUObjectArray.CloseDisregardForGC();
+	}
 }
 
 void CleanupCoreUObject()
