Add -waitfordebugger argument matching to low level test runner to match editor/game arguments

diff --git a/Engine/Source/Developer/LowLevelTestsRunner/Private/TestRunner.cpp b/Engine/Source/Developer/LowLevelTestsRunner/Private/TestRunner.cpp
--- a/Engine/Source/Developer/LowLevelTestsRunner/Private/TestRunner.cpp
+++ b/Engine/Source/Developer/LowLevelTestsRunner/Private/TestRunner.cpp
@@ -160,6 +160,10 @@ void FTestRunner::ParseCommandLine(TConstArrayView<const ANSICHAR*> Args)
 		{
 			bAttachToDebugger = true;
 		}
+		else if (Arg == ANSITEXTVIEW("--waitfordebugger"))
+		{
+			bAttachToDebugger = true;
+		}
 		else if (Arg == ANSITEXTVIEW("--buildmachine"))
 		{
 			GIsBuildMachine = true;
