Add explicit dependency to Core in LowLevelTestsRunner and fix missing symbols if ApplicationCore is not included for iOS

diff --git a/Engine/Source/Developer/LowLevelTestsRunner/LowLevelTestsRunner.Build.cs b/Engine/Source/Developer/LowLevelTestsRunner/LowLevelTestsRunner.Build.cs
--- a/Engine/Source/Developer/LowLevelTestsRunner/LowLevelTestsRunner.Build.cs
+++ b/Engine/Source/Developer/LowLevelTestsRunner/LowLevelTestsRunner.Build.cs
@@ -22,9 +22,9 @@ namespace UnrealBuildTool.Rules
 
 			PrivateDependencyModuleNames.AddRange(
 				new string[] {
-					/* Add any extra deps here, should not require any UE modules */
+					"Core"
 				}
-			);
+			);		
 		}
 	}
 }
diff --git a/Engine/Source/Developer/LowLevelTestsRunner/Private/Platform/IOS/IOSTestRunner.cpp b/Engine/Source/Developer/LowLevelTestsRunner/Private/Platform/IOS/IOSTestRunner.cpp
--- a/Engine/Source/Developer/LowLevelTestsRunner/Private/Platform/IOS/IOSTestRunner.cpp
+++ b/Engine/Source/Developer/LowLevelTestsRunner/Private/Platform/IOS/IOSTestRunner.cpp
@@ -4,38 +4,8 @@
 
 #include "TestRunner.h"
 #include "Platform/Apple/AppleTestRunnerHelper.h"
-#include "Containers/UnrealString.h"
 #include <mach-o/dyld.h>
 
-// Empty implementation of missing symbols from ApplicationCore
-APPLICATIONCORE_API FString GSavedCommandLine;
-
-@class IOSAppDelegate;
-@class UIApplication;
-
-namespace FAppEntry
-{
-	void PlatformInit() {}
-	void PreInit(IOSAppDelegate* AppDelegate, UIApplication* Application) {}
-	void Init() {}
-	void Tick() {}
-    void SuspendTick() {}
-	void ResumeAudioContext() {}
-	void ResetAudioContextResumeTime() {}
-	void Shutdown() {}
-    void Suspend(bool bIsInterrupt = false) {}
-    void Resume(bool bIsInterrupt = false) {}
-	void RestartAudio() {}
-    void IncrementAudioSuspendCounters() {}
-    void DecrementAudioSuspendCounters() {}
-
-	bool IsStartupMoviePlaying() { return false; }
-
-	bool	gAppLaunchedWithLocalNotification;
-	FString	gLaunchLocalNotificationActivationEvent;
-	int32	gLaunchLocalNotificationFireDate;
-}
-
 const char* GetCacheDirectory()
 {
 	return "/private/var/tmp/unit_tests_pds_cache";
diff --git a/Engine/Source/Developer/LowLevelTestsRunner/Private/Platform/IOS/MissingModuleSymbols.cpp b/Engine/Source/Developer/LowLevelTestsRunner/Private/Platform/IOS/MissingModuleSymbols.cpp
new file mode 100644
--- /dev/null
+++ b/Engine/Source/Developer/LowLevelTestsRunner/Private/Platform/IOS/MissingModuleSymbols.cpp
@@ -0,0 +1,67 @@
+// Copyright Epic Games, Inc. All Rights Reserved.
+
+#if PLATFORM_IOS
+
+#if WITH_APPLICATION_CORE
+#include "Containers/UnrealString.h"
+
+// Empty implementation of missing symbols used in ApplicationCore but implemented somewhere else (in Launch for instance)
+FString GSavedCommandLine;
+
+@class IOSAppDelegate;
+@class UIApplication;
+
+namespace FAppEntry
+{
+	void PlatformInit() {}
+	void PreInit(IOSAppDelegate* AppDelegate, UIApplication* Application) {}
+	void Init() {}
+	void Tick() {}
+    void SuspendTick() {}
+	void ResumeAudioContext() {}
+	void ResetAudioContextResumeTime() {}
+	void Shutdown() {}
+    void Suspend(bool bIsInterrupt = false) {}
+    void Resume(bool bIsInterrupt = false) {}
+	void RestartAudio() {}
+    void IncrementAudioSuspendCounters() {}
+    void DecrementAudioSuspendCounters() {}
+
+	bool IsStartupMoviePlaying() { return false; }
+
+	bool	gAppLaunchedWithLocalNotification;
+	FString	gLaunchLocalNotificationActivationEvent;
+	int32	gLaunchLocalNotificationFireDate;
+}
+#else // !WITH_APPLICATION_CORE
+
+// Empty implementation of missing symbols used in Core but implemented in ApplicationCore
+#include "Logging/LogMacros.h"
+#include "HAL/PlatformMisc.h"
+
+DEFINE_LOG_CATEGORY_STATIC(LogLowLevelTests, Log, VeryVerbose);
+
+// IOSAppDelegate is mostly used in IOSPlatformMisc implementation. It is always accessed using [IOSAppDelegate GetDelegate]
+@interface IOSAppDelegate : NSObject
+{
+}
+@end
+
+@implementation IOSAppDelegate
+// Implementation returns nil so any method call on that nil will become a no-op. Not having this method would produce a 
+// "Sending unrecognized selector to object" uncaught exception in any method in FIOSPlatformMisc that tries to get the delegate  
++ (IOSAppDelegate*)GetDelegate
+{
+	UE_LOG(LogLowLevelTests, Error, TEXT("Attempt to access to IOSAppDelegate methods. Those are only available if ApplicationCore is present"));
+	return nil;
+}
+@end
+
+EAppReturnType::Type MessageBoxExtImpl( EAppMsgType::Type MsgType, const TCHAR* Text, const TCHAR* Caption )
+{
+	return EAppReturnType::Type::Yes;
+}
+
+#endif
+
+#endif
