LowLevelTests: Add ability to set -ProjectDir via first paramter ending in .uproject or via -projectdir.

This should allow running LowLevelTests correctly via UnrealVS.

diff --git a/Engine/Source/Developer/LowLevelTestsRunner/Private/TestRunner.cpp b/Engine/Source/Developer/LowLevelTestsRunner/Private/TestRunner.cpp
--- a/Engine/Source/Developer/LowLevelTestsRunner/Private/TestRunner.cpp
+++ b/Engine/Source/Developer/LowLevelTestsRunner/Private/TestRunner.cpp
@@ -8,9 +8,11 @@
 #include "Logging/LogSuppressionInterface.h"
 #include "Misc/ScopeExit.h"
 #include "Misc/StringBuilder.h"
+#include "Misc/Paths.h"
 #include "Modules/ModuleManager.h"
 #include "String/Find.h"
 #include "String/LexFromString.h"
+#include "Containers/UnrealString.h"
 #include "TestCommon/CoreUtilities.h"
 #include "TestRunnerOutputDeviceError.h"
 
@@ -90,6 +92,34 @@ FTestRunner::FTestRunner()
 void FTestRunner::ParseCommandLine(TConstArrayView<const ANSICHAR*> Args)
 {
 	bool bExtraArg = false;
+	bool FirstItemIsUproject = false;
+
+
+	// If program is launched from UnrealVS the project directoy is passed in as the first arugment without a -.
+	// We will need to parse that out and treat it as our project directory. Programs may have project directories diffrent
+	// from their binaries directory.
+	if (Args.Num() > 2)
+	{
+		FAnsiStringView FirstArg = Args[1];
+		FirstItemIsUproject = !FirstArg.StartsWith(ANSITEXTVIEW("-")) && FirstArg.EndsWith(ANSITEXTVIEW(".uproject"));
+	}
+
+	bool bFirstItemIsUproject = false;
+	bool bFirstArugumentIsProjectName = false;
+	bool bFirstArgument = true;
+
+	if (Args.Num() > 2)
+	{
+		FAnsiStringView FirstArg = Args[1];
+		// If program is launched from UnrealVS the project directoy is passed in as the first arugment without a -.
+		// We will need to parse that out and treat it as our project directory. Programs may have project directories diffrent
+		// from their binaries directory.
+		bFirstItemIsUproject = !FirstArg.StartsWith(ANSITEXTVIEW("-")) && FirstArg.EndsWith(ANSITEXTVIEW(".uproject"));
+
+		//If a program is launched from the staged directory the first argument may be the program name
+		bFirstArugumentIsProjectName = !FirstArg.StartsWith(ANSITEXTVIEW("-")) && FirstArg.Compare(GInternalProjectName) == 0;
+	}
+
 	for (FAnsiStringView Arg : Args)
 	{
 		if (bExtraArg)
@@ -112,6 +142,32 @@ void FTestRunner::ParseCommandLine(TConstArrayView<const ANSICHAR*> Args)
 		{
 			bExtraArg = true;
 		}
+		else if (bFirstArgument && bFirstArugumentIsProjectName)
+		{
+			// Do nothing
+		}
+		else if (bFirstArgument && bFirstItemIsUproject || Arg.StartsWith(ANSITEXTVIEW("--projectdir=")))
+		{
+			FStringBuilderBase Builder;
+			FString ProjectDirOverride;
+			if (bFirstItemIsUproject) {
+				Builder.Append(Arg);
+				ProjectDirOverride = FPaths::GetPath(*Builder);
+			}
+			else
+			{
+				Builder.Append(Arg.RightChop(13));
+				ProjectDirOverride = *Builder;
+			}
+
+			if (!ProjectDirOverride.EndsWith(TEXT("/")))
+			{
+				ProjectDirOverride += TEXT("/");
+			}
+
+			FPaths::NormalizeFilename(ProjectDirOverride);
+			FGenericPlatformMisc::SetOverrideProjectDir(ProjectDirOverride);
+		}
 		else if (Arg.StartsWith(ANSITEXTVIEW("--sleep=")))
 		{
 			LexFromString(SleepOnInitSeconds, WriteToString<16>(Arg.RightChop(8)).ToView());
@@ -172,6 +228,8 @@ void FTestRunner::ParseCommandLine(TConstArrayView<const ANSICHAR*> Args)
 		{
 			CatchArgs.Add(Arg.GetData());
 		}
+
+		bFirstArgument = false;
 	}
 
 	// Break in the debugger on failed assertions when attached.
