LowLevelTests: Fix issue with nested && || warning with GCC.

diff --git a/Engine/Source/Developer/LowLevelTestsRunner/Private/TestRunner.cpp b/Engine/Source/Developer/LowLevelTestsRunner/Private/TestRunner.cpp
--- a/Engine/Source/Developer/LowLevelTestsRunner/Private/TestRunner.cpp
+++ b/Engine/Source/Developer/LowLevelTestsRunner/Private/TestRunner.cpp
@@ -146,7 +146,7 @@ void FTestRunner::ParseCommandLine(TConstArrayView<const ANSICHAR*> Args)
 		{
 			// Do nothing
 		}
-		else if (bFirstArgument && bFirstItemIsUproject || Arg.StartsWith(ANSITEXTVIEW("--projectdir=")))
+		else if (bFirstArgument && (bFirstItemIsUproject || Arg.StartsWith(ANSITEXTVIEW("--projectdir="))))
 		{
 			FStringBuilderBase Builder;
 			FString ProjectDirOverride;
