LowLevelTestRunner: Fix bug when both .uproject and --projectdir is present, projectdir is incorrectly parsed as a .uproject path instead of a directory path.

diff --git a/Engine/Source/Developer/LowLevelTestsRunner/Private/TestRunner.cpp b/Engine/Source/Developer/LowLevelTestsRunner/Private/TestRunner.cpp
--- a/Engine/Source/Developer/LowLevelTestsRunner/Private/TestRunner.cpp
+++ b/Engine/Source/Developer/LowLevelTestsRunner/Private/TestRunner.cpp
@@ -150,7 +150,7 @@ void FTestRunner::ParseCommandLine(TConstArrayView<const ANSICHAR*> Args)
 		{
 			FStringBuilderBase Builder;
 			FString ProjectDirOverride;
-			if (bFirstItemIsUproject) {
+			if (bFirstArgument && bFirstItemIsUproject) {
 				Builder.Append(Arg);
 				ProjectDirOverride = FPaths::GetPath(*Builder);
 			}
