LowLevelTests: Fix executable name being used as project path instead of the .uprogram when present.

diff --git a/Engine/Source/Developer/LowLevelTestsRunner/Private/TestRunner.cpp b/Engine/Source/Developer/LowLevelTestsRunner/Private/TestRunner.cpp
--- a/Engine/Source/Developer/LowLevelTestsRunner/Private/TestRunner.cpp
+++ b/Engine/Source/Developer/LowLevelTestsRunner/Private/TestRunner.cpp
@@ -106,7 +106,7 @@ void FTestRunner::ParseCommandLine(TConstArrayView<const ANSICHAR*> Args)
 
 	bool bFirstItemIsUproject = false;
 	bool bFirstArugumentIsProjectName = false;
-	bool bFirstArgument = true;
+	bool bFirstArgument = false;
 
 	if (Args.Num() > 2)
 	{
@@ -122,6 +122,8 @@ void FTestRunner::ParseCommandLine(TConstArrayView<const ANSICHAR*> Args)
 
 	for (FAnsiStringView Arg : Args)
 	{
+		// Track args[1] as the first argument, as args[0] is always the program name.
+		bFirstArgument = Arg == Args[1];
 		if (bExtraArg)
 		{
 			if (const int32 SpaceIndex = String::FindFirstChar(Arg, ' '); SpaceIndex != INDEX_NONE)
