UnrealVS - do not project prepend project path for Editor configuration in the case of test targets
- Project generation will add a new IsTestTarget property in the "Globals" property group section of the project file for explicit test target projects
- UnrealVS checks for this property and does not change the command line for test targets
- TestRunner removes all project argument logic - that can be passed after --extra-args using -project="Path/To/*.uproject" and/or -projectdir="Path/To/Project"
- OSS and Online global handlers, which occur during the Catch2 session, will read project name and/or directory using SetProjectNameAndDirectory from TestCommon/Initialization.h

diff --git a/Engine/Extras/UnrealVS/VS2022/AtomFeed.xml b/Engine/Extras/UnrealVS/VS2022/AtomFeed.xml
--- a/Engine/Extras/UnrealVS/VS2022/AtomFeed.xml
+++ b/Engine/Extras/UnrealVS/VS2022/AtomFeed.xml
@@ -8,7 +8,7 @@
     <title type="text">UnrealVS</title>
     <summary type="text">Unreal Engine extension for Visual Studio</summary>
     <published>2022-10-21T10:01:57-07:00</published>
-    <updated>2023-06-02T00:06:00-00:00</updated>
+    <updated>2023-07-25T15:44:40-00:00</updated>
     <author>
       <name>Epic Games, Inc.</name>
     </author>
@@ -19,7 +19,7 @@
     <content type="application/octet-stream" src="UnrealVS.vsix" />
     <Vsix xmlns:xsd="http://www.w3.org/2001/XMLSchema" xmlns:xsi="http://www.w3.org/2001/XMLSchema-instance" xmlns="http://schemas.microsoft.com/developer/vsx-syndication-schema/2010">
       <Id>ddbf523f-7eb6-4887-bd51-85a714ff87eb</Id>
-      <Version>1.84.3.2</Version>
+      <Version>1.84.3.3</Version>
       <Rating xsi:nil="true" />
       <RatingCount xsi:nil="true" />
       <DownloadCount xsi:nil="true" />
diff --git a/Engine/Source/Developer/LowLevelTestsRunner/Private/TestCommon/Initialization.cpp b/Engine/Source/Developer/LowLevelTestsRunner/Private/TestCommon/Initialization.cpp
--- a/Engine/Source/Developer/LowLevelTestsRunner/Private/TestCommon/Initialization.cpp
+++ b/Engine/Source/Developer/LowLevelTestsRunner/Private/TestCommon/Initialization.cpp
@@ -5,9 +5,13 @@
 #include "TestCommon/CoreUObjectUtilities.h"
 #include "TestCommon/EngineUtilities.h"
 
+#include "Containers/UnrealString.h"
 #include "GenericPlatform/GenericPlatformFile.h"
 #include "HAL/PlatformFileManager.h"
+#include "Misc/CommandLine.h"
 #include "Misc/DelayedAutoRegister.h"
+#include "Misc/Paths.h"
+
 
 static IPlatformFile* DefaultPlatformFile;
 
@@ -50,6 +54,40 @@ void UseDefaultPlatformFile()
 	FPlatformFileManager::Get().SetPlatformFile(*DefaultPlatformFile);
 }
 
+void SetProjectNameAndDirectory()
+{
+	FString ProjectFileOrName;
+	FString ProjectDirOverride;
+
+	bool bIsProjectNamePassed = false;
+
+	FParse::Value(FCommandLine::Get(), TEXT("-project="), ProjectFileOrName);
+	if (!ProjectFileOrName.IsEmpty())
+	{
+		bIsProjectNamePassed = !ProjectFileOrName.EndsWith(TEXT(".uproject"));
+	}
+
+	FParse::Value(FCommandLine::Get(), TEXT("-projectdir="), ProjectDirOverride);
+	if (!ProjectDirOverride.IsEmpty())
+	{
+		if (!ProjectDirOverride.EndsWith(TEXT("/")))
+		{
+			ProjectDirOverride.Append(TEXT("/"));
+		}
+	}
+
+	if (!bIsProjectNamePassed && ProjectDirOverride.IsEmpty() && !ProjectFileOrName.IsEmpty())
+	{
+		ProjectDirOverride = FPaths::GetPath(ProjectFileOrName);
+	}
+	
+	if (!ProjectDirOverride.IsEmpty())
+	{
+		FPaths::NormalizeDirectoryName(ProjectDirOverride);
+		FGenericPlatformMisc::SetOverrideProjectDir(ProjectDirOverride);
+	}
+}
+
 void InitAll(bool bAllowLogging, bool bMultithreaded)
 {
 	SaveDefaultPlatformFile();
diff --git a/Engine/Source/Developer/LowLevelTestsRunner/Private/TestRunner.cpp b/Engine/Source/Developer/LowLevelTestsRunner/Private/TestRunner.cpp
--- a/Engine/Source/Developer/LowLevelTestsRunner/Private/TestRunner.cpp
+++ b/Engine/Source/Developer/LowLevelTestsRunner/Private/TestRunner.cpp
@@ -92,38 +92,9 @@ FTestRunner::FTestRunner()
 void FTestRunner::ParseCommandLine(TConstArrayView<const ANSICHAR*> Args)
 {
 	bool bExtraArg = false;
-	bool FirstItemIsUproject = false;
-
-
-	// If program is launched from UnrealVS the project directoy is passed in as the first arugment without a -.
-	// We will need to parse that out and treat it as our project directory. Programs may have project directories diffrent
-	// from their binaries directory.
-	if (Args.Num() > 2)
-	{
-		FAnsiStringView FirstArg = Args[1];
-		FirstItemIsUproject = !FirstArg.StartsWith(ANSITEXTVIEW("-")) && FirstArg.EndsWith(ANSITEXTVIEW(".uproject"));
-	}
-
-	bool bFirstItemIsUproject = false;
-	bool bFirstArugumentIsProjectName = false;
-	bool bFirstArgument = false;
-
-	if (Args.Num() > 2)
-	{
-		FAnsiStringView FirstArg = Args[1];
-		// If program is launched from UnrealVS the project directoy is passed in as the first arugment without a -.
-		// We will need to parse that out and treat it as our project directory. Programs may have project directories diffrent
-		// from their binaries directory.
-		bFirstItemIsUproject = !FirstArg.StartsWith(ANSITEXTVIEW("-")) && FirstArg.EndsWith(ANSITEXTVIEW(".uproject"));
-
-		//If a program is launched from the staged directory the first argument may be the program name
-		bFirstArugumentIsProjectName = !FirstArg.StartsWith(ANSITEXTVIEW("-")) && FirstArg.Compare(GInternalProjectName) == 0;
-	}
 
 	for (FAnsiStringView Arg : Args)
 	{
-		// Track args[1] as the first argument, as args[0] is always the program name.
-		bFirstArgument = Arg == Args[1];
 		if (bExtraArg)
 		{
 			if (const int32 SpaceIndex = String::FindFirstChar(Arg, ' '); SpaceIndex != INDEX_NONE)
@@ -144,32 +115,6 @@ void FTestRunner::ParseCommandLine(TConstArrayView<const ANSICHAR*> Args)
 		{
 			bExtraArg = true;
 		}
-		else if (bFirstArgument && bFirstArugumentIsProjectName)
-		{
-			// Do nothing
-		}
-		else if ((bFirstArgument && bFirstItemIsUproject) || Arg.StartsWith(ANSITEXTVIEW("--projectdir=")))
-		{
-			FStringBuilderBase Builder;
-			FString ProjectDirOverride;
-			if (bFirstArgument && bFirstItemIsUproject) {
-				Builder.Append(Arg);
-				ProjectDirOverride = FPaths::GetPath(*Builder);
-			}
-			else
-			{
-				Builder.Append(Arg.RightChop(13));
-				ProjectDirOverride = *Builder;
-			}
-
-			if (!ProjectDirOverride.EndsWith(TEXT("/")))
-			{
-				ProjectDirOverride += TEXT("/");
-			}
-
-			FPaths::NormalizeFilename(ProjectDirOverride);
-			FGenericPlatformMisc::SetOverrideProjectDir(ProjectDirOverride);
-		}
 		else if (Arg.StartsWith(ANSITEXTVIEW("--sleep=")))
 		{
 			LexFromString(SleepOnInitSeconds, WriteToString<16>(Arg.RightChop(8)).ToView());
@@ -230,8 +175,6 @@ void FTestRunner::ParseCommandLine(TConstArrayView<const ANSICHAR*> Args)
 		{
 			CatchArgs.Add(Arg.GetData());
 		}
-
-		bFirstArgument = false;
 	}
 
 	// Break in the debugger on failed assertions when attached.
diff --git a/Engine/Source/Developer/LowLevelTestsRunner/Public/TestCommon/Initialization.h b/Engine/Source/Developer/LowLevelTestsRunner/Public/TestCommon/Initialization.h
--- a/Engine/Source/Developer/LowLevelTestsRunner/Public/TestCommon/Initialization.h
+++ b/Engine/Source/Developer/LowLevelTestsRunner/Public/TestCommon/Initialization.h
@@ -8,6 +8,8 @@ void SaveDefaultPlatformFile();
 void UsePlatformFileStubIfRequired();
 void UseDefaultPlatformFile();
 
+void SetProjectNameAndDirectory();
+
 void InitAllThreadPoolsEditorEx(bool MultiThreaded);
 
 void InitOutputDevicesEx();
diff --git a/Engine/Source/Programs/UnrealBuildTool/ProjectFiles/VisualStudio/VCProject.cs b/Engine/Source/Programs/UnrealBuildTool/ProjectFiles/VisualStudio/VCProject.cs
--- a/Engine/Source/Programs/UnrealBuildTool/ProjectFiles/VisualStudio/VCProject.cs
+++ b/Engine/Source/Programs/UnrealBuildTool/ProjectFiles/VisualStudio/VCProject.cs
@@ -981,11 +981,19 @@ namespace UnrealBuildTool
 				}
 			}
 
+			ProjectConfigAndTargetCombination? FoundCombo = ProjectConfigAndTargetCombinations.FirstOrDefault(combo => combo != null && combo.ProjectTarget != null && combo.ProjectTarget.TargetRules != null);
+			TargetRules? DefaultRules = FoundCombo != null ? FoundCombo.ProjectTarget?.TargetRules : null;
+			bool IsTestTarget = (DefaultRules != null ? DefaultRules.IsTestTarget : false);
+
 			// Project globals (project GUID, project type, SCC bindings, etc)
 			{
 				VCProjectFileContent.AppendLine("  <PropertyGroup Label=\"Globals\">");
 				VCProjectFileContent.AppendLine("    <ProjectGuid>{0}</ProjectGuid>", ProjectGUID.ToString("B").ToUpperInvariant());
 				VCProjectFileContent.AppendLine("    <RootNamespace>{0}</RootNamespace>", ProjectName);
+				if (IsTestTarget)
+				{
+					VCProjectFileContent.AppendLine("    <IsTestTarget>true</IsTestTarget>");
+				}
 				VCProjectFileContent.AppendLine("  </PropertyGroup>");
 			}
 
@@ -1310,9 +1318,6 @@ namespace UnrealBuildTool
 				VCPreprocessorDefinitions.Append(CurDef);
 			}
 
-			ProjectConfigAndTargetCombination? FoundCombo = ProjectConfigAndTargetCombinations.FirstOrDefault(combo => combo != null && combo.ProjectTarget != null && combo.ProjectTarget.TargetRules != null);
-			TargetRules? DefaultRules = FoundCombo != null ? FoundCombo.ProjectTarget?.TargetRules : null;
-
 			var GetAdditionalOptionsString = (TargetRules? TargetRules) =>
 			{
 				return string.Format("{0} {1}{2}{3}", GetCppStandardCompileArgument(GetIntelliSenseCppVersion()),
diff --git a/Engine/Source/Programs/UnrealVS/ReleaseNotes.txt b/Engine/Source/Programs/UnrealVS/ReleaseNotes.txt
--- a/Engine/Source/Programs/UnrealVS/ReleaseNotes.txt
+++ b/Engine/Source/Programs/UnrealVS/ReleaseNotes.txt
@@ -1,3 +1,6 @@
+1.84.3.3
+* Bug fix for low level test projects - do not auto prefix project path in command line
+
 1.84.3.2
 * Bug fix for compiling a single file when us
 
diff --git a/Engine/Source/Programs/UnrealVS/UnrealVS.2019/Properties/AssemblyInfo.cs b/Engine/Source/Programs/UnrealVS/UnrealVS.2019/Properties/AssemblyInfo.cs
--- a/Engine/Source/Programs/UnrealVS/UnrealVS.2019/Properties/AssemblyInfo.cs
+++ b/Engine/Source/Programs/UnrealVS/UnrealVS.2019/Properties/AssemblyInfo.cs
@@ -31,6 +31,6 @@ using System.Runtime.InteropServices;
 // You can specify all the values or you can default the Revision and Build Numbers 
 // by using the '*' as shown below:
 
-[assembly: AssemblyVersion("1.84.3.2")]
-[assembly: AssemblyFileVersion("1.84.3.2")]
+[assembly: AssemblyVersion("1.84.3.3")]
+[assembly: AssemblyFileVersion("1.84.3.3")]
 
diff --git a/Engine/Source/Programs/UnrealVS/UnrealVS.Shared/CommandLineEditor.cs b/Engine/Source/Programs/UnrealVS/UnrealVS.Shared/CommandLineEditor.cs
--- a/Engine/Source/Programs/UnrealVS/UnrealVS.Shared/CommandLineEditor.cs
+++ b/Engine/Source/Programs/UnrealVS/UnrealVS.Shared/CommandLineEditor.cs
@@ -253,7 +253,10 @@ namespace UnrealVS
 							}
 							// for "Game" projects automatically remove the game project filename from the start of the command line
 							var ActiveConfiguration = (SolutionConfiguration2)UnrealVSPackage.Instance.DTE.Solution.SolutionBuild.ActiveConfiguration;
-							if (UnrealVSPackage.Instance.IsUESolutionLoaded && Utils.IsGameProject(SelectedStartupProject) && Utils.HasUProjectCommandLineArg(ActiveConfiguration.Name))
+							if (UnrealVSPackage.Instance.IsUESolutionLoaded &&
+								Utils.IsGameProject(SelectedStartupProject) &&
+								!Utils.IsTestTargetProject(SelectedStartupProject) &&
+								Utils.HasUProjectCommandLineArg(ActiveConfiguration.Name))
 							{
 								string AutoPrefix = Utils.GetAutoUProjectCommandLinePrefix(SelectedStartupProject);
 								if (!string.IsNullOrEmpty(AutoPrefix))
@@ -348,7 +351,10 @@ namespace UnrealVS
 
 							// for "Game" projects automatically remove the game project filename from the start of the command line
 							var ActiveConfiguration = (SolutionConfiguration2)UnrealVSPackage.Instance.DTE.Solution.SolutionBuild.ActiveConfiguration;
-							if (UnrealVSPackage.Instance.IsUESolutionLoaded && Utils.IsGameProject(SelectedStartupProject) && Utils.HasUProjectCommandLineArg(ActiveConfiguration.Name))
+							if (UnrealVSPackage.Instance.IsUESolutionLoaded &&
+								Utils.IsGameProject(SelectedStartupProject) &&
+								!Utils.IsTestTargetProject(SelectedStartupProject) &&
+								Utils.HasUProjectCommandLineArg(ActiveConfiguration.Name))
 							{
 								string AutoPrefix = Utils.GetAutoUProjectCommandLinePrefix(SelectedStartupProject);
 								if (FullCommandLine.IndexOf(Utils.UProjectExtension, StringComparison.OrdinalIgnoreCase) < 0 &&
diff --git a/Engine/Source/Programs/UnrealVS/UnrealVS.Shared/Utils.cs b/Engine/Source/Programs/UnrealVS/UnrealVS.Shared/Utils.cs
--- a/Engine/Source/Programs/UnrealVS/UnrealVS.Shared/Utils.cs
+++ b/Engine/Source/Programs/UnrealVS/UnrealVS.Shared/Utils.cs
@@ -452,6 +452,16 @@ namespace UnrealVS
 			return GetUProjects().ContainsKey(Project.Name);
 		}
 
+		public static bool IsTestTargetProject(Project Project)
+		{
+			ThreadHelper.ThrowIfNotOnUIThread();
+			if (Project.Globals.VariableExists["IsTestTarget"])
+			{
+				return Convert.ToBoolean(Project.Globals["IsTestTarget"]);
+			}
+			return false;
+		}
+
 		/// <summary>
 		/// Does the config build something that takes a .uproject on the command line?
 		/// </summary>
