Fixed an Assertion fail

diff --git a/Engine/Source/Developer/LowLevelTestsRunner/Private/TestCommon/Initialization.cpp b/Engine/Source/Developer/LowLevelTestsRunner/Private/TestCommon/Initialization.cpp
--- a/Engine/Source/Developer/LowLevelTestsRunner/Private/TestCommon/Initialization.cpp
+++ b/Engine/Source/Developer/LowLevelTestsRunner/Private/TestCommon/Initialization.cpp
@@ -84,6 +84,10 @@ void SetProjectNameAndDirectory()
 	if (!ProjectDirOverride.IsEmpty())
 	{
 		FPaths::NormalizeDirectoryName(ProjectDirOverride);
+		if (!ProjectDirOverride.EndsWith(TEXT("/")))
+		{
+			ProjectDirOverride.Append(TEXT("/"));
+		}
 		FGenericPlatformMisc::SetOverrideProjectDir(ProjectDirOverride);
 	}
 }
