Improve tests log output when an assertion is triggered so it shows both the expression and the expanded expression if possible
Add Catch::StringConvetrer specialization for FOnlineError in OnlineTests and improved checks to be more informative

diff --git a/Engine/Source/Developer/LowLevelTestsRunner/Private/TestListeners/EventListener.cpp b/Engine/Source/Developer/LowLevelTestsRunner/Private/TestListeners/EventListener.cpp
--- a/Engine/Source/Developer/LowLevelTestsRunner/Private/TestListeners/EventListener.cpp
+++ b/Engine/Source/Developer/LowLevelTestsRunner/Private/TestListeners/EventListener.cpp
@@ -116,11 +116,24 @@ public:
 	{
 		if (ITestRunner* TestRunner = ITestRunner::Get())
 		{
-			UE_CLOG(TestRunner->HasLogOutput() && !AssertionStats.assertionResult.succeeded(),
-				LogLowLevelTests, Error, TEXT("Assertion \"%hs\" failed at %hs(%" SIZE_T_FMT ")"),
-				AssertionStats.assertionResult.getExpression().c_str(),
-				AssertionStats.assertionResult.getSourceInfo().file,
-				SIZE_T(AssertionStats.assertionResult.getSourceInfo().line));
+			const bool bShouldLogAssertion = TestRunner->HasLogOutput() && !AssertionStats.assertionResult.succeeded();
+			if (AssertionStats.assertionResult.hasExpandedExpression())
+			{
+				UE_CLOG(bShouldLogAssertion,
+					LogLowLevelTests, Error, TEXT("Assertion \"%hs\" failed with \"%hs\" at %hs(%" SIZE_T_FMT ")"),
+					AssertionStats.assertionResult.getExpression().c_str(),
+					AssertionStats.assertionResult.getExpandedExpression().c_str(),
+					AssertionStats.assertionResult.getSourceInfo().file,
+					SIZE_T(AssertionStats.assertionResult.getSourceInfo().line));
+			}
+			else
+			{
+				UE_CLOG(bShouldLogAssertion,
+					LogLowLevelTests, Error, TEXT("Assertion \"%hs\" failed at %hs(%" SIZE_T_FMT ")"),
+					AssertionStats.assertionResult.getExpression().c_str(),
+					AssertionStats.assertionResult.getSourceInfo().file,
+					SIZE_T(AssertionStats.assertionResult.getSourceInfo().line));
+			}
 		}
 	}
 };
