LLT: Fix no stdout for Linux tests

diff --git a/Engine/Source/Developer/LowLevelTestsRunner/Private/Platform/Linux/LinuxTestRunner.cpp b/Engine/Source/Developer/LowLevelTestsRunner/Private/Platform/Linux/LinuxTestRunner.cpp
--- a/Engine/Source/Developer/LowLevelTestsRunner/Private/Platform/Linux/LinuxTestRunner.cpp
+++ b/Engine/Source/Developer/LowLevelTestsRunner/Private/Platform/Linux/LinuxTestRunner.cpp
@@ -2,6 +2,12 @@
 
 #if defined(PLATFORM_LINUX)
 #include "TestRunner.h"
+
+#include "Misc/OutputDeviceDebug.h"
+#include "Misc/OutputDeviceRedirector.h"
+#include "HAL/PlatformOutputDevices.h"
+#include "CoreGlobals.h"
+
 #include <unistd.h>
 
 const char* GetProcessExecutablePath()
@@ -33,8 +39,14 @@ const char* GetCacheDirectory()
 	return "/var/tmp/playground_test_pds_cache";
 }
 
+void InitializeLinuxPlatform()
+{
+	GLog->AddOutputDevice(new FOutputDeviceDebug());
+}
+
 int main(int argc, const char* argv[])
 {
+	FTestDelegates::GetGlobalPlatformSetup().BindStatic(&InitializeLinuxPlatform);
 	return RunTests(argc, argv);
 }
 #endif
\ No newline at end of file
diff --git a/Engine/Source/Developer/LowLevelTestsRunner/Private/TestRunner.cpp b/Engine/Source/Developer/LowLevelTestsRunner/Private/TestRunner.cpp
--- a/Engine/Source/Developer/LowLevelTestsRunner/Private/TestRunner.cpp
+++ b/Engine/Source/Developer/LowLevelTestsRunner/Private/TestRunner.cpp
@@ -242,10 +242,10 @@ void FTestRunner::GlobalSetup()
 	{
 		// Set up default output devices to handle Log, Verbose, VeryVerbose.
 		FPlatformOutputDevices::SetupOutputDevices();
-
 		FLogSuppressionInterface::Get().ProcessConfigAndCommandLine();
 	}
 
+	FTestDelegates::GetGlobalPlatformSetup().ExecuteIfBound();
 	FTestDelegates::GetGlobalSetup().ExecuteIfBound();
 }
 
diff --git a/Engine/Source/Developer/LowLevelTestsRunner/Public/TestRunner.h b/Engine/Source/Developer/LowLevelTestsRunner/Public/TestRunner.h
--- a/Engine/Source/Developer/LowLevelTestsRunner/Public/TestRunner.h
+++ b/Engine/Source/Developer/LowLevelTestsRunner/Public/TestRunner.h
@@ -29,6 +29,7 @@
 const char* GetCacheDirectory();
 
 DECLARE_DELEGATE(FGlobalSetupHandler);
+DECLARE_DELEGATE(FGlobalPlatformSetupHandler);
 
 class FTestDelegates
 {
@@ -44,6 +45,12 @@ public:
 		static TUniquePtr<FGlobalSetupHandler> GlobalTeardown = MakeUnique<FGlobalSetupHandler>();
 		return *GlobalTeardown.Get();
 	}
+
+	static FGlobalPlatformSetupHandler& GetGlobalPlatformSetup()
+	{
+		static TUniquePtr<FGlobalPlatformSetupHandler> GlobalPlatformSetup = MakeUnique<FGlobalPlatformSetupHandler>();
+		return *GlobalPlatformSetup.Get();
+	}
 };
 
 int RunTests(int argc, const char* argv[]);
\ No newline at end of file
