Improve logs for LowLeveltests by including captured/scoped info data when assertions are triggered

diff --git a/Engine/Source/Developer/LowLevelTestsRunner/Private/TestListeners/EventListener.cpp b/Engine/Source/Developer/LowLevelTestsRunner/Private/TestListeners/EventListener.cpp
--- a/Engine/Source/Developer/LowLevelTestsRunner/Private/TestListeners/EventListener.cpp
+++ b/Engine/Source/Developer/LowLevelTestsRunner/Private/TestListeners/EventListener.cpp
@@ -117,22 +117,31 @@ public:
 		if (ITestRunner* TestRunner = ITestRunner::Get())
 		{
 			const bool bShouldLogAssertion = TestRunner->HasLogOutput() && !AssertionStats.assertionResult.succeeded();
-			if (AssertionStats.assertionResult.hasExpandedExpression())
+			if (bShouldLogAssertion)
 			{
-				UE_CLOG(bShouldLogAssertion,
-					LogLowLevelTests, Error, TEXT("Assertion \"%hs\" failed with \"%hs\" at %hs(%" SIZE_T_FMT ")"),
-					AssertionStats.assertionResult.getExpression().c_str(),
-					AssertionStats.assertionResult.getExpandedExpression().c_str(),
-					AssertionStats.assertionResult.getSourceInfo().file,
-					SIZE_T(AssertionStats.assertionResult.getSourceInfo().line));
-			}
-			else
-			{
-				UE_CLOG(bShouldLogAssertion,
-					LogLowLevelTests, Error, TEXT("Assertion \"%hs\" failed at %hs(%" SIZE_T_FMT ")"),
-					AssertionStats.assertionResult.getExpression().c_str(),
-					AssertionStats.assertionResult.getSourceInfo().file,
-					SIZE_T(AssertionStats.assertionResult.getSourceInfo().line));
+				if (AssertionStats.assertionResult.hasExpandedExpression())
+				{
+					UE_LOG( LogLowLevelTests, Error, TEXT("Assertion \"%hs\" failed with \"%hs\" at %hs(%" SIZE_T_FMT ")"),
+							AssertionStats.assertionResult.getExpression().c_str(),
+							AssertionStats.assertionResult.getExpandedExpression().c_str(),
+							AssertionStats.assertionResult.getSourceInfo().file,
+							SIZE_T(AssertionStats.assertionResult.getSourceInfo().line));
+				}
+				else
+				{
+					UE_LOG( LogLowLevelTests, Error, TEXT("Assertion \"%hs\" failed at %hs(%" SIZE_T_FMT ")"),
+							AssertionStats.assertionResult.getExpression().c_str(),
+							AssertionStats.assertionResult.getSourceInfo().file,
+							SIZE_T(AssertionStats.assertionResult.getSourceInfo().line));
+				}
+
+				for (const Catch::MessageInfo& Info : AssertionStats.infoMessages)
+				{
+					UE_LOG( LogLowLevelTests, Warning, TEXT("Info: \"%hs\" at %hs(%" SIZE_T_FMT ")"),
+							Info.message.c_str(),
+							Info.lineInfo.file,
+							SIZE_T(Info.lineInfo.line));
+				}
 			}
 		}
 	}
