Fix compilation error when building implcit tests

diff --git a/Engine/Source/Developer/LowLevelTestsRunner/README.md b/Engine/Source/Developer/LowLevelTestsRunner/README.md
--- a/Engine/Source/Developer/LowLevelTestsRunner/README.md
+++ b/Engine/Source/Developer/LowLevelTestsRunner/README.md
@@ -520,22 +520,24 @@ You can build and run explicit tests directly from Visual Studio on desktop plat
 
 #### Build
 
+All low level tests must be built usint the UBT mode `-Mode=Test`. This mode performs a pre-analysis of the build target's dependency chain and detects UE core modules such as CoreUObject, Engine, ApplicationCore and UnrealEd and sets appropriate compilation flags. By default `-Mode=Test` is used to build explicit tests - tests that have their own explicitely defined target. To build an implicit test from an existing target use `-Mode=Test -Implicit`: note the additional -Implicit argument here.
+
 ##### Explicit Tests
 
-You can use Unreal Build Tool to build explicit tests. Suppose we build explicit test cases with their target class `MyTestsTarget`:
+You can use Unreal Build Tool to build explicit tests using `-Mode=Test`. Suppose we build explicit test cases with their target class `MyTestsTarget`:
 
 ```
-.\Engine\Binaries\DotNET\UnrealBuildTool\UnrealBuildTool.exe MyTestsTarget Development Win64
+.\RunUBT.bat MyTestsTarget Development Win64 -Mode=Test
 ```
 
 The configuration used above is `Development` and the platform is `Win64` for example purposes. All configurations and platforms are supported.
 
 ##### Implicit Tests
 
-To build implicit tests, use an existing target, for example `UnrealEditor`, and use `-Mode=Test`, which builds a program target based on the given target:
+To build implicit tests, use an existing target, for example `UnrealEditor`, and use `-Mode=Test -Implicit`, which builds a program target based on the given target:
 
 ```
-.\Engine\Binaries\DotNET\UnrealBuildTool\UnrealBuildTool.exe UnrealEditor Development Win64 -Mode=Test
+.\RunUBT.bat UnrealEditor Development Win64 -Mode=Test -Implicit
 ```
 
 Using test mode, the tests that are included in every module in the dependency graph are collected into one executable program. Unreal Build Tool is currently the only way to build implicit tests.
diff --git a/Engine/Source/Programs/UnrealBuildTool/Configuration/TestModuleRules.cs b/Engine/Source/Programs/UnrealBuildTool/Configuration/TestModuleRules.cs
--- a/Engine/Source/Programs/UnrealBuildTool/Configuration/TestModuleRules.cs
+++ b/Engine/Source/Programs/UnrealBuildTool/Configuration/TestModuleRules.cs
@@ -69,6 +69,8 @@ namespace UnrealBuildTool
 			PrivateDependencyModuleNames.AddRange(TestedModule.PrivateDependencyModuleNames);
 			PublicDependencyModuleNames.AddRange(TestedModule.PublicDependencyModuleNames);
 
+			DirectoriesForModuleSubClasses = new Dictionary<Type, DirectoryReference>();
+
 			// Tests can refer to tested module's Public and Private paths
 			string ModulePublicDir = Path.Combine(TestedModule.ModuleDirectory, "Public");
 			if (System.IO.Directory.Exists(ModulePublicDir))
