Tear down GLog prior to static destruction

This avoids crashes from anything that logs during static destruction in the test runner, which can happen as the test runner only manually tears down a subset of things that the standard LaunchEngineLoop handles

#rnx

diff --git a/Engine/Source/Developer/LowLevelTestsRunner/Private/TestRunner.cpp b/Engine/Source/Developer/LowLevelTestsRunner/Private/TestRunner.cpp
--- a/Engine/Source/Developer/LowLevelTestsRunner/Private/TestRunner.cpp
+++ b/Engine/Source/Developer/LowLevelTestsRunner/Private/TestRunner.cpp
@@ -9,6 +9,7 @@
 #include "Misc/ScopeExit.h"
 #include "Misc/StringBuilder.h"
 #include "Misc/Paths.h"
+#include "Misc/OutputDeviceRedirector.h"
 #include "Modules/ModuleManager.h"
 #include "String/Find.h"
 #include "String/LexFromString.h"
@@ -312,8 +313,14 @@ int RunTests(int32 ArgC, const ANSICHAR* ArgV[])
 	{
 		TestRunner.GlobalTeardown();
 		TestRunner.Terminate();
+		
 		FModuleManager::Get().UnloadModulesAtShutdown();
 		RequestEngineExit(TEXT("Exiting"));
+		
+		if (GLog)
+		{
+			GLog->TearDown();
+		}
 	};
 
 	int CatchReturn = TestRunner.RunCatchSession();
