Tear down localization prior to static destruction

This avoids crashes from a logging attempt during static destruction in the test runner, which can happen as the test runner only manually tears down a subset of things that the standard LaunchEngineLoop handles

#rnx

diff --git a/Engine/Source/Developer/LowLevelTestsRunner/Private/TestRunner.cpp b/Engine/Source/Developer/LowLevelTestsRunner/Private/TestRunner.cpp
--- a/Engine/Source/Developer/LowLevelTestsRunner/Private/TestRunner.cpp
+++ b/Engine/Source/Developer/LowLevelTestsRunner/Private/TestRunner.cpp
@@ -321,6 +321,9 @@ int RunTests(int32 ArgC, const ANSICHAR* ArgV[])
 		{
 			GLog->TearDown();
 		}
+
+		FTextLocalizationManager::TearDown();
+		FInternationalization::TearDown();
 	};
 
 	int CatchReturn = TestRunner.RunCatchSession();
