Move UE component cleanup logic from TestRunner into proper place (introduced at CL 32031322) and fix refactoring (CL 32031169) on Sony platforms to re-allow unpackaged builds

diff --git a/Engine/Plugins/Experimental/StateGraph/Tests/StateGraphTests/Private/TestGroupEvents.cpp b/Engine/Plugins/Experimental/StateGraph/Tests/StateGraphTests/Private/TestGroupEvents.cpp
new file mode 100644
--- /dev/null
+++ b/Engine/Plugins/Experimental/StateGraph/Tests/StateGraphTests/Private/TestGroupEvents.cpp
@@ -0,0 +1,10 @@
+// Copyright Epic Games, Inc. All Rights Reserved.
+#include "TestCommon/Initialization.h"
+
+#include <catch2/catch_test_macros.hpp>
+
+GROUP_AFTER_GLOBAL(Catch::DefaultGroup)
+{
+	CleanupLogs();
+	CleanupLocalization();
+}
diff --git a/Engine/Source/Developer/LowLevelTestsRunner/Private/TestCommon/Initialization.cpp b/Engine/Source/Developer/LowLevelTestsRunner/Private/TestCommon/Initialization.cpp
--- a/Engine/Source/Developer/LowLevelTestsRunner/Private/TestCommon/Initialization.cpp
+++ b/Engine/Source/Developer/LowLevelTestsRunner/Private/TestCommon/Initialization.cpp
@@ -10,6 +10,7 @@
 #include "HAL/PlatformFileManager.h"
 #include "Misc/CommandLine.h"
 #include "Misc/DelayedAutoRegister.h"
+#include "Misc/OutputDeviceRedirector.h"
 #include "Misc/Paths.h"
 
 
@@ -117,6 +118,20 @@ void InitAll(bool bAllowLogging, bool bMultithreaded)
 	GIsRunning = true;
 }
 
+void CleanupLogs()
+{
+	if (GLog)
+	{
+		GLog->TearDown();
+	}
+}
+
+void CleanupLocalization()
+{
+	FTextLocalizationManager::TearDown();
+	FInternationalization::TearDown();
+}
+
 void CleanupAll()
 {
 #if WITH_ENGINE
@@ -127,4 +142,6 @@ void CleanupAll()
 #endif
 	CleanupAllThreadPools();
 	CleanupTaskGraph();
+	CleanupLogs();
+	CleanupLocalization();
 }
\ No newline at end of file
diff --git a/Engine/Source/Developer/LowLevelTestsRunner/Private/TestRunner.cpp b/Engine/Source/Developer/LowLevelTestsRunner/Private/TestRunner.cpp
--- a/Engine/Source/Developer/LowLevelTestsRunner/Private/TestRunner.cpp
+++ b/Engine/Source/Developer/LowLevelTestsRunner/Private/TestRunner.cpp
@@ -9,7 +9,6 @@
 #include "Misc/ScopeExit.h"
 #include "Misc/StringBuilder.h"
 #include "Misc/Paths.h"
-#include "Misc/OutputDeviceRedirector.h"
 #include "Modules/ModuleManager.h"
 #include "String/Find.h"
 #include "String/LexFromString.h"
@@ -316,14 +315,6 @@ int RunTests(int32 ArgC, const ANSICHAR* ArgV[])
 		
 		FModuleManager::Get().UnloadModulesAtShutdown();
 		RequestEngineExit(TEXT("Exiting"));
-		
-		if (GLog)
-		{
-			GLog->TearDown();
-		}
-
-		FTextLocalizationManager::TearDown();
-		FInternationalization::TearDown();
 	};
 
 	int CatchReturn = TestRunner.RunCatchSession();
diff --git a/Engine/Source/Developer/LowLevelTestsRunner/Public/TestCommon/Initialization.h b/Engine/Source/Developer/LowLevelTestsRunner/Public/TestCommon/Initialization.h
--- a/Engine/Source/Developer/LowLevelTestsRunner/Public/TestCommon/Initialization.h
+++ b/Engine/Source/Developer/LowLevelTestsRunner/Public/TestCommon/Initialization.h
@@ -18,4 +18,6 @@ void InitStats();
 
 void InitAll(bool bAllowLogging, bool bMultithreaded);
 
+void CleanupLogs();
+void CleanupLocalization();
 void CleanupAll();
\ No newline at end of file
diff --git a/Engine/Source/Programs/LowLevelTests/FoundationTests/Tests/TestGroupEvents.cpp b/Engine/Source/Programs/LowLevelTests/FoundationTests/Tests/TestGroupEvents.cpp
--- a/Engine/Source/Programs/LowLevelTests/FoundationTests/Tests/TestGroupEvents.cpp
+++ b/Engine/Source/Programs/LowLevelTests/FoundationTests/Tests/TestGroupEvents.cpp
@@ -1,23 +1,14 @@
 // Copyright Epic Games, Inc. All Rights Reserved.
-#include "TestCommon/CoreUtilities.h"
-#include "TestCommon/CoreUObjectUtilities.h"
+#include "TestCommon/Initialization.h"
 
 #include <catch2/catch_test_macros.hpp>
 
 GROUP_BEFORE_GLOBAL(Catch::DefaultGroup)
 {
-	InitTaskGraph();
-	InitThreadPool(true);
-#if WITH_COREUOBJECT
-	InitCoreUObject();
-#endif
+	InitAll(true, true);
 }
 
 GROUP_AFTER_GLOBAL(Catch::DefaultGroup)
 {
-#if WITH_COREUOBJECT
-	CleanupCoreUObject();
-#endif
-	CleanupThreadPool();
-	CleanupTaskGraph();
+	CleanupAll();
 }
diff --git a/Engine/Source/Programs/LowLevelTests/MathCoreTests/Private/TestGroupEvents.cpp b/Engine/Source/Programs/LowLevelTests/MathCoreTests/Private/TestGroupEvents.cpp
new file mode 100644
--- /dev/null
+++ b/Engine/Source/Programs/LowLevelTests/MathCoreTests/Private/TestGroupEvents.cpp
@@ -0,0 +1,10 @@
+// Copyright Epic Games, Inc. All Rights Reserved.
+#include "TestCommon/Initialization.h"
+
+#include <catch2/catch_test_macros.hpp>
+
+GROUP_AFTER_GLOBAL(Catch::DefaultGroup)
+{
+	CleanupLogs();
+	CleanupLocalization();
+}
diff --git a/Engine/Source/Programs/WebTests/Private/WebLLTModule.cpp b/Engine/Source/Programs/WebTests/Private/WebLLTModule.cpp
--- a/Engine/Source/Programs/WebTests/Private/WebLLTModule.cpp
+++ b/Engine/Source/Programs/WebTests/Private/WebLLTModule.cpp
@@ -7,6 +7,15 @@
 #include "Modules/ModuleManager.h"
 #include "TestRunner.h"
 
+#include "TestCommon/Initialization.h"
+#include <catch2/catch_test_macros.hpp>
+
+GROUP_AFTER_GLOBAL(Catch::DefaultGroup)
+{
+	CleanupLogs();
+	CleanupLocalization();
+}
+
 namespace UE::WebTests
 {
 
diff --git a/Engine/Source/Runtime/Online/Experimental/EventLoopTests/Tests/TestGroupEvents.cpp b/Engine/Source/Runtime/Online/Experimental/EventLoopTests/Tests/TestGroupEvents.cpp
new file mode 100644
--- /dev/null
+++ b/Engine/Source/Runtime/Online/Experimental/EventLoopTests/Tests/TestGroupEvents.cpp
@@ -0,0 +1,10 @@
+// Copyright Epic Games, Inc. All Rights Reserved.
+#include "TestCommon/Initialization.h"
+
+#include <catch2/catch_test_macros.hpp>
+
+GROUP_AFTER_GLOBAL(Catch::DefaultGroup)
+{
+	CleanupLogs();
+	CleanupLocalization();
+}
