OSS Tests PS4 - Test application didn't log any test events after timeout period of 5 minutes, stopping.

diff --git a/Engine/Source/Developer/LowLevelTestsRunner/Private/TestCommon/Initialization.cpp b/Engine/Source/Developer/LowLevelTestsRunner/Private/TestCommon/Initialization.cpp
--- a/Engine/Source/Developer/LowLevelTestsRunner/Private/TestCommon/Initialization.cpp
+++ b/Engine/Source/Developer/LowLevelTestsRunner/Private/TestCommon/Initialization.cpp
@@ -57,6 +57,10 @@ void UseDefaultPlatformFile()
 
 void SetProjectNameAndDirectory()
 {
+	// There may be cases on some platforms that you need to verify files are in the filesystem, 
+	// and this call was needed to correctly initalize the system.
+	FPaths::ProjectDir();
+
 	FString ProjectFileOrName;
 	FString ProjectDirOverride;
 
