LLT Android missing script that fxes manifest re-added

diff --git a/Engine/Source/Developer/LowLevelTestsRunner/LowLevelTestsRunner.Build.cs b/Engine/Source/Developer/LowLevelTestsRunner/LowLevelTestsRunner.Build.cs
--- a/Engine/Source/Developer/LowLevelTestsRunner/LowLevelTestsRunner.Build.cs
+++ b/Engine/Source/Developer/LowLevelTestsRunner/LowLevelTestsRunner.Build.cs
@@ -42,6 +42,12 @@ namespace UnrealBuildTool.Rules
 			{
 				PrivateDependencyModuleNames.Add("CoreUObject");
 			}
+
+			if (Target.Platform == UnrealTargetPlatform.Android)
+			{
+				string ModulePath = Utils.MakePathRelativeTo(ModuleDirectory, Target.RelativeEnginePath);
+				AdditionalPropertiesForReceipt.Add("AndroidPlugin", Path.Combine(ModulePath, "LowLevelTestsAndroid_UPL.xml"));
+			}
 		}
 	}
 }
