Platform file stub not loading - stub module name changed in CL 33938220 but invoked with previous name

diff --git a/Engine/Source/Developer/LowLevelTestsRunner/Private/TestCommon/Initialization.cpp b/Engine/Source/Developer/LowLevelTestsRunner/Private/TestCommon/Initialization.cpp
--- a/Engine/Source/Developer/LowLevelTestsRunner/Private/TestCommon/Initialization.cpp
+++ b/Engine/Source/Developer/LowLevelTestsRunner/Private/TestCommon/Initialization.cpp
@@ -36,7 +36,7 @@ void InitStats()
 void UsePlatformFileStubIfRequired()
 {
 #if UE_LLT_USE_PLATFORM_FILE_STUB
-	if (IPlatformFile* WrapperFile = FPlatformFileManager::Get().GetPlatformFile(TEXT("PlatformFileStub")))
+	if (IPlatformFile* WrapperFile = FPlatformFileManager::Get().GetPlatformFile(TEXT("LowLevelTestsRunner")))
 	{
 		IPlatformFile* CurrentPlatformFile = &FPlatformFileManager::Get().GetPlatformFile();
 		WrapperFile->Initialize(CurrentPlatformFile, TEXT(""));
diff --git a/Engine/Source/Developer/LowLevelTestsRunner/Private/TestStubs/PlatformFileModuleStub.cpp b/Engine/Source/Developer/LowLevelTestsRunner/Private/TestStubs/PlatformFileModuleStub.cpp
--- a/Engine/Source/Developer/LowLevelTestsRunner/Private/TestStubs/PlatformFileModuleStub.cpp
+++ b/Engine/Source/Developer/LowLevelTestsRunner/Private/TestStubs/PlatformFileModuleStub.cpp
@@ -14,7 +14,7 @@ public:
 	virtual bool Initialize(IPlatformFile* Inner, const TCHAR* CmdLine) override { LowerLevel = Inner; return true; }
 	virtual IPlatformFile* GetLowerLevel() { return LowerLevel; }
 	virtual void SetLowerLevel(IPlatformFile* NewLowerLevel) { LowerLevel = NewLowerLevel; }
-	virtual const TCHAR* GetName() const { return TEXT("LLTPlatformFileStub"); }
+	virtual const TCHAR* GetName() const { return TEXT("PlatformFileStub"); }
 	virtual bool FileExists(const TCHAR* Filename) override { return false; }
 	virtual int64 FileSize(const TCHAR* Filename) override { return -1LL; }
 	virtual bool DeleteFile(const TCHAR* Filename) override { return false; }
