Fix copy-pasta in AssertionsMacros.h
#rnx
#tests LowLevelTests EnginePreflight

diff --git a/Engine/Source/Developer/LowLevelTestsRunner/Public/TestMacros/Assertions.h b/Engine/Source/Developer/LowLevelTestsRunner/Public/TestMacros/Assertions.h
--- a/Engine/Source/Developer/LowLevelTestsRunner/Public/TestMacros/Assertions.h
+++ b/Engine/Source/Developer/LowLevelTestsRunner/Public/TestMacros/Assertions.h
@@ -18,10 +18,10 @@
 #define CHECK_ENSURE_MSG(msg, ...) INTERNAL_UE_ENSURE_MSG(msg, "CHECK_ENSURE", DO_ENSURE, Catch::ResultDisposition::ContinueOnFailure, #__VA_ARGS__, __VA_ARGS__ )
 
 //requires that a UE `check` fails in this call
-#define REQUIRE_CHECK(...) INTERNAL_UE_CHECK( "REQUIRE_ENSURE", DO_CHECK, Catch::ResultDisposition::Normal, #__VA_ARGS__, __VA_ARGS__ )
+#define REQUIRE_CHECK(...) INTERNAL_UE_CHECK( "REQUIRE_CHECK", DO_CHECK, Catch::ResultDisposition::Normal, #__VA_ARGS__, __VA_ARGS__ )
 
 //requires that a UE `check` fails in this call containg the supplies message
-#define REQUIRE_CHECK_MSG(msg, ...) INTERNAL_UE_CHECK_MSG(msg, "REQUIRE_ENSURE", DO_CHECK, Catch::ResultDisposition::Normal, #__VA_ARGS__, __VA_ARGS__ )
+#define REQUIRE_CHECK_MSG(msg, ...) INTERNAL_UE_CHECK_MSG(msg, "REQUIRE_CHECK", DO_CHECK, Catch::ResultDisposition::Normal, #__VA_ARGS__, __VA_ARGS__ )
 
 #define INTERNAL_UE_ENSURE( macroName, doEnsure, resultDisposition, ensureExpr, ... ) \
 	do { \
