Update LowLevelTests.xml to pass in -programprojectoverride=$(BuildProject) when using StagesWithProjectFile=True
Add support for LowLevelTests uprojects in /Tests/ folder within Game Folders

diff --git a/Engine/Build/LowLevelTests.xml b/Engine/Build/LowLevelTests.xml
--- a/Engine/Build/LowLevelTests.xml
+++ b/Engine/Build/LowLevelTests.xml
@@ -136,11 +136,14 @@
 	
 	<Macro Name="BuildAndStageTest" Arguments="Platform;TestName;StagesWithProjectFile;TargetName;ExtraCompilationArgs">
 		<Property Name="LocalAchiveDirectory" Value="$(RootDir)/$(TargetName)"/>
+		<Property Name="ExtraCompilationArgsFinal" Value="$(ExtraCompilationArgs)"/>
 		<Do If="'$(StagesWithProjectFile)' == 'True'">
-			<Command Name="BuildCookRun" Arguments="-Build -NoBootstrapExe -NoCodeSign -SkipCook -Stage -NoSubmit -Archive -ArchiveDirectory=$(LocalAchiveDirectory) -project=$(TargetName) -configuration=$(Configuration) -platform=$(Platform) $(ExtraCompilationArgs)" If="!$(DryRun)"/>
+			<Do If="'$(BuildProject)' != ''">
+				<Property Name="ExtraCompilationArgsFinal" Value="$(ExtraCompilationArgs) -projectprogramoverride=$(RootDir)/$(BuildProject)"/>
+			</Do>
+			<Command Name="BuildCookRun" Arguments="-Build -NoBootstrapExe -NoCodeSign -SkipCook -SkipBuildEditor=True -Stage -NoSubmit -Archive -ArchiveDirectory=$(LocalAchiveDirectory) -project=$(TargetName) -configuration=$(Configuration) -platform=$(Platform) $(ExtraCompilationArgsFinal)" If="!$(DryRun)"/>
 		</Do>
 		<Do If="'$(StagesWithProjectFile)' != 'True'">
-			<Property Name="ExtraCompilationArgsFinal" Value="$(ExtraCompilationArgs)"/>
 			<Do If="'$(BuildProject)' != ''">
 				<Property Name="ExtraCompilationArgsFinal" Value="$(ExtraCompilationArgs) -Project=$(RootDir)/$(BuildProject)"/>
 			</Do>
diff --git a/Engine/Source/Developer/LowLevelTestsRunner/Private/TestCommon/Initialization.cpp b/Engine/Source/Developer/LowLevelTestsRunner/Private/TestCommon/Initialization.cpp
--- a/Engine/Source/Developer/LowLevelTestsRunner/Private/TestCommon/Initialization.cpp
+++ b/Engine/Source/Developer/LowLevelTestsRunner/Private/TestCommon/Initialization.cpp
@@ -57,6 +57,13 @@ void UseDefaultPlatformFile()
 
 void SetProjectNameAndDirectory()
 {
+	// Use target name instead of project file being passed in as this will be more accruate to finding the program's config
+	// and project files. Resolves issues with UE_PROJECT_NAME being set to the parent project instead of the program's project.
+	// This will likely break if we ever have one LLT project file have mutliple targets eg: LyraGameEOS vs LyraGame
+#ifdef UE_TARGET_NAME
+	FCString::Strncpy(GInternalProjectName, TEXT(PREPROCESSOR_TO_STRING(UE_TARGET_NAME)), UE_ARRAY_COUNT(GInternalProjectName));
+#endif
+
 	// There may be cases on some platforms that you need to verify files are in the filesystem, 
 	// and this call was needed to correctly initalize the system.
 	FPaths::ProjectDir();
diff --git a/Engine/Source/Developer/LowLevelTestsRunner/README.md b/Engine/Source/Developer/LowLevelTestsRunner/README.md
--- a/Engine/Source/Developer/LowLevelTestsRunner/README.md
+++ b/Engine/Source/Developer/LowLevelTestsRunner/README.md
@@ -252,6 +252,25 @@ To add implicit tests to this module, create a directory titled **Tests** at the
 ```
 With this structure, the file `SampleCodeTests.cpp` contains tests for the `SampleCode.cpp` file, and `MoreSampleCodeTests.cpp` contains tests for the `MoreSampleCode.cpp` file.
 
+Should your tests depend on Game-like configuration loading creating a SampleCodeTests.uproject in the root of the will allow the Test program to find the configuration folder with an additional `--projectdir=` paramter. 
+
+```
+.
+└── SampleModule/
+    ├── Public/
+    ├── Private/
+    └── Tests/
+        ├── SampleCodeTests.cpp
+        └── SampleModuleSubDirectory/
+            └── MoreSampleCodeTests.cpp
+        └── Config/
+            └── DefaultEngine.ini
+        └── SampleCodeTests.uproject
+```
+
+Note the `.Build.cs` should be updated to include `StagesWithProjectFile=True` to indicate to BuildGraph that we should use a Game-like folder structure when deploying these tests to run.
+
+
 ## Write Low-Level Tests
 
 This section primarily discusses structure, guidelines, and best practices for writing **Low-Level Tests (LLTs)** with Catch2 in the context of **Unreal Engine (UE)**. See the [Catch2 GitHub Repository]([REDACTED_PROFILE_URL]/Catch2) for information specific to Catch2. For a complete guide on writing tests, see the [Catch2 Reference]([REDACTED_PROFILE_URL]/Catch2/tree/devel/docs).
@@ -533,6 +552,7 @@ You can build and run low-level tests with:
 * [Visual Studio](#visual-studio)
 * [Unreal Build Tool (UBT)](#unreal-build-tool)
 * [BuildGraph](#buildgraph)
+* [BuildCookRun](#buildcookrun)
 
 You can build and run explicit tests using any of these tools. We recommend you build and run explicit tests with BuildGraph whenever possible. Currently, you can only build and run implicit tests with Unreal Build Tool.
 
@@ -672,6 +692,27 @@ Here are some common ways to use the BuildGraph script:
 .\RunUAT.bat BuildGraph -Script="Engine/Build/LowLevelTests.xml" -Target="Catch2 Build Library Win64"
 ```
 
+### BuildCookRun
+
+Engine provided script BuildCookRun provides a wrapping script to Unreal Build Tool allows you to build, cook, and run your application in one script. For workflows specific to only deploying low level tests and running them we recommend the following command. BuildCookRun can be used on Low Level Test programs that have `StagesWithProjectFile=true` specified in its metadata and .uproject files appropriately configured. 
+
+It is important to pass `--projectprogramoverride` if your tests exist inside a Game's project folder for dependencies to correctly be passed to Unreal Build Tool.
+
+```
+.\RunUAT.bat BuildCookRun  -Build -NoBootstrapExe -NoCodeSign -SkipCook -Stage -NoSubmit -SkipBuildEditor=True -SkipClient -Archive -ArchiveDirectory=D:\build\SampleTestName -project=MyGame\Path\To\SampleTestName.uproject -configuration=Development -platform=Win64 --projectprogramoverride=MyGame\MyGame.uproject
+```
+
+If your Binaries folder within the Archive directory appears  subfolder inside your Game/Binaries/ folder, you will want to include a directory remapping so the correct Project folder will be found as well, otherwise the default UE5 system will not find content and configs.
+
+```
+[Staging]
+;Remap platform folders for loading configs from these specific folders
++RemapDirectories=(From="SampleGame/Binaries/Win64/SampleTests",To="SampleTests/Binaries/Win64")
+```
+
+If we want to deploy to a device the addional `-Device=<Device.json>` paramter can be passed in for BuildCookRun to deploy your program to a device. For additional help on BuildCookRun see `-Help` or specific documenation pages on Automation Tool.
+
+
 ## Example: Foundation Tests
 
 ### Overview
diff --git a/Engine/Source/Programs/AutomationTool/AutomationUtils/ProjectUtils.cs b/Engine/Source/Programs/AutomationTool/AutomationUtils/ProjectUtils.cs
--- a/Engine/Source/Programs/AutomationTool/AutomationUtils/ProjectUtils.cs
+++ b/Engine/Source/Programs/AutomationTool/AutomationUtils/ProjectUtils.cs
@@ -380,6 +380,7 @@ namespace AutomationTool
 
 			List<DirectoryReference> GameFolders = new List<DirectoryReference>();
 			DirectoryReference RulesFolder = new DirectoryReference(GetRulesAssemblyFolder());
+			List<DirectoryReference> ExtraSearchDirectories = (ExtraSearchPaths == null) ? null : ExtraSearchPaths.Select(x => new DirectoryReference(x)).ToList();
 			if (Properties.RawProjectPath != null)
 			{
 				Logger.LogDebug("Looking for targets for project {Arg0}", Properties.RawProjectPath);
@@ -388,6 +389,12 @@ namespace AutomationTool
 
 				FullProjectPath = CommandUtils.GetDirectoryName(Properties.RawProjectPath.FullName).Replace("\\", "/");
 
+				// Extra search root uproject folder if we have /tests/
+				if (FullProjectPath.Contains("/Tests/"))
+				{
+					ExtraSearchDirectories.Add(new DirectoryReference(FullProjectPath));
+				}
+
 				// there is a special case of Programs, where the uproject doesn't align with the Source directory, so we redirect to where
 				// the program's target.cs file(s) are
 				if (FullProjectPath.Contains("/Programs/"))
@@ -411,7 +418,7 @@ namespace AutomationTool
 				CommandUtils.PushDir(SourceDir);
 				DirPushed = true;
 			}
-			List<DirectoryReference> ExtraSearchDirectories = (ExtraSearchPaths == null)? null : ExtraSearchPaths.Select(x => new DirectoryReference(x)).ToList();
+
 			List<FileReference> TargetScripts = Rules.FindAllRulesSourceFiles(Rules.RulesFileType.Target, GameFolders: GameFolders, ForeignPlugins: null, AdditionalSearchPaths: ExtraSearchDirectories);
 			if (DirPushed)
 			{
