Added REQUIRE_CHECK_SLOW to low-level test macros, that validates that checkSlow, checkfSlow... are called on configurations where they're effective

#tests FoundationTests debug+development

diff --git a/Engine/Source/Developer/LowLevelTestsRunner/Public/TestMacros/Assertions.h b/Engine/Source/Developer/LowLevelTestsRunner/Public/TestMacros/Assertions.h
--- a/Engine/Source/Developer/LowLevelTestsRunner/Public/TestMacros/Assertions.h
+++ b/Engine/Source/Developer/LowLevelTestsRunner/Public/TestMacros/Assertions.h
@@ -20,9 +20,15 @@
 //requires that a UE `check` fails in this call
 #define REQUIRE_CHECK(...) INTERNAL_UE_CHECK( "REQUIRE_CHECK", DO_CHECK, Catch::ResultDisposition::Normal, #__VA_ARGS__, __VA_ARGS__ )
 
-//requires that a UE `check` fails in this call containg the supplies message
+//requires that a UE `check` fails in this call contains the supplies message
 #define REQUIRE_CHECK_MSG(msg, ...) INTERNAL_UE_CHECK_MSG(msg, "REQUIRE_CHECK", DO_CHECK, Catch::ResultDisposition::Normal, #__VA_ARGS__, __VA_ARGS__ )
 
+//requires that a UE `checkSlow` fails in this call
+#define REQUIRE_CHECK_SLOW(...) INTERNAL_UE_CHECK( "REQUIRE_CHECK_SLOW", DO_GUARD_SLOW, Catch::ResultDisposition::Normal, #__VA_ARGS__, __VA_ARGS__ )
+
+//requires that a UE `checkSlow` fails in this call contains the supplies message
+#define REQUIRE_CHECK_SLOW_MSG(msg, ...) INTERNAL_UE_CHECK_MSG(msg, "REQUIRE_CHECK_SLOW", DO_GUARD_SLOW, Catch::ResultDisposition::Normal, #__VA_ARGS__, __VA_ARGS__ )
+
 #define INTERNAL_UE_ENSURE( macroName, doEnsure, resultDisposition, ensureExpr, ... ) \
 	do { \
 		Catch::AssertionHandler catchAssertionHandler( macroName##_catch_sr, CATCH_INTERNAL_LINEINFO, ensureExpr, resultDisposition ); \
@@ -31,7 +37,7 @@
 			static_cast<void>(__VA_ARGS__); \
 			bool bEncounteredEnsure = scope.GetCount() > 0; \
 			if (doEnsure && !bEncounteredEnsure) \
-				catchAssertionHandler.handleMessage( Catch::ResultWas::ExplicitFailure ,"Expected failure of `ensure` not recieved."); \
+				catchAssertionHandler.handleMessage( Catch::ResultWas::ExplicitFailure ,"Expected failure of `ensure` not received."); \
 		} INTERNAL_CATCH_CATCH( catchAssertionHandler ) \
 		INTERNAL_CATCH_REACT( catchAssertionHandler ) \
 	} while(false) \
@@ -44,7 +50,7 @@
 			static_cast<void>(__VA_ARGS__); \
 			bool bEncounteredEnsure = scope.GetCount() > 0; \
 			if (doEnsure && !bEncounteredEnsure) \
-				catchAssertionHandler.handleMessage( Catch::ResultWas::ExplicitFailure, ( Catch::MessageStream() << "Expected failure of `ensure` with message '" << msg << "' not recieved" + ::Catch::StreamEndStop() ).m_stream.str() ); \
+				catchAssertionHandler.handleMessage( Catch::ResultWas::ExplicitFailure, ( Catch::MessageStream() << "Expected failure of `ensure` with message '" << msg << "' not received" + ::Catch::StreamEndStop() ).m_stream.str() ); \
 		} INTERNAL_CATCH_CATCH( catchAssertionHandler ) \
 		INTERNAL_CATCH_REACT( catchAssertionHandler ) \
 	} while(false) \
@@ -57,7 +63,7 @@
 			__VA_ARGS__; \
 			bool bEncounteredEnsure = scope.GetCount() > 0; \
 			if (doCheck && !bEncounteredEnsure) \
-				catchAssertionHandler.handleMessage( Catch::ResultWas::ExplicitFailure, "Expected failure of `check` not recieved"); \
+				catchAssertionHandler.handleMessage( Catch::ResultWas::ExplicitFailure, "Expected failure of `check` not received"); \
 		} INTERNAL_CATCH_CATCH( catchAssertionHandler ) \
 		INTERNAL_CATCH_REACT( catchAssertionHandler ) \
 	} while(false) \
@@ -70,7 +76,7 @@
 			__VA_ARGS__; \
 			bool bEncounteredEnsure = scope.GetCount() > 0; \
 			if (doCheck && !bEncounteredEnsure) \
-				catchAssertionHandler.handleMessage( Catch::ResultWas::ExplicitFailure, ( Catch::MessageStream() << "Expected failure of `check` containing message '" << msg << "' not recieved" + ::Catch::StreamEndStop() ).m_stream.str() ); \
+				catchAssertionHandler.handleMessage( Catch::ResultWas::ExplicitFailure, ( Catch::MessageStream() << "Expected failure of `check` containing message '" << msg << "' not received" + ::Catch::StreamEndStop() ).m_stream.str() ); \
 		} INTERNAL_CATCH_CATCH( catchAssertionHandler ) \
 		INTERNAL_CATCH_REACT( catchAssertionHandler ) \
 	} while(false) \
diff --git a/Engine/Source/Runtime/Core/Tests/Misc/AssertionMacrosTest.cpp b/Engine/Source/Runtime/Core/Tests/Misc/AssertionMacrosTest.cpp
--- a/Engine/Source/Runtime/Core/Tests/Misc/AssertionMacrosTest.cpp
+++ b/Engine/Source/Runtime/Core/Tests/Misc/AssertionMacrosTest.cpp
@@ -123,4 +123,11 @@ TEST_CASE("Core::Misc::TestCheck", "[Core][Misc][AssertionMacros][Check]")
 	REQUIRE_CHECK_MSG("Error Message", checkf(1 == 2, TEXT("Error Message")));
 }
 
+TEST_CASE("Core::Misc::TestCheck", "[Core][Misc][AssertionMacros][CheckSlow]")
+{
+	REQUIRE_CHECK_SLOW(checkSlow(1 == 2));
+	REQUIRE_CHECK_SLOW_MSG("1 == 2", checkfSlow(1 == 2, TEXT("Error Message")));
+	REQUIRE_CHECK_SLOW_MSG("Error Message", checkfSlow(1 == 2, TEXT("Error Message")));
+}
+
 #endif
