LLT doc update
#rnx

diff --git a/Engine/Source/Developer/LowLevelTestsRunner/README.md b/Engine/Source/Developer/LowLevelTestsRunner/README.md
--- a/Engine/Source/Developer/LowLevelTestsRunner/README.md
+++ b/Engine/Source/Developer/LowLevelTestsRunner/README.md
@@ -211,6 +211,24 @@ public UEModuleTestsTarget(TargetInfo Target) : base(Target)
 ```
 </font>
 
+It is also necessary to initialize and cleanup Engine components.
+To do this, provide a .cpp file containing initialization events:
+<font size="2">
+```cpp
+#include "TestCommon/Initialization.h"
+
+GROUP_BEFORE_ALL(Catch::DefaultGroup)
+{
+  InitAll(true, true);
+}
+
+GROUP_AFTER_ALL(Catch::DefaultGroup)
+{
+  CleanupAll();
+}
+```
+</font>
+
 ## Implicit Tests
 
 Implicit Tests reside in the same module that you want to test. All implicit tests are conditionally compiled through the `WITH_LOW_LEVEL_TESTS` definition. While they can be placed anywhere in the module, for example, in the **Private** folder, it can be hard to distinguish them from module source code. To ensure that your tests are discoverable, follow these conventions:
@@ -401,7 +419,7 @@ The examples contained in this section are not exhaustive of all the features of
 
 ### More Examples
 
-There are several UE-specific Low-Level Test examples in the engine directory `Engine/Source/Runtime/Core/Tests`. To continue the example from the [Types of Low-Level Tests](#types-of-low-level-tests), you can see an example of TDD-style tests in the file `ArchiveReplaceObjectRefTests.cpp` located in `Engine/Source/Programs/LowLevelTests/Tests`.
+There are several UE-specific Low-Level Test examples in the engine directory `Engine/Source/Runtime/Core/Tests`. To continue the example from the [Types of Low-Level Tests](#types-of-low-level-tests), you can see an example of TDD-style tests in the file `ArchiveReplaceObjectRefTests.cpp` located in `Engine/Source/Programs/LowLevelTests/FoundationTests/Tests`.
 
 ## Additional Low-Level Tests Features
 
