[misc] even more FString::Printf format fixes

diff --git a/Engine/Source/Developer/LowLevelTestsRunner/Public/TestCommon/Expectations.h b/Engine/Source/Developer/LowLevelTestsRunner/Public/TestCommon/Expectations.h
--- a/Engine/Source/Developer/LowLevelTestsRunner/Public/TestCommon/Expectations.h
+++ b/Engine/Source/Developer/LowLevelTestsRunner/Public/TestCommon/Expectations.h
@@ -33,7 +33,7 @@ namespace
 	{
 		if (Value.GetData())
 		{
-			return FString::Printf(TEXT("\"%.*hs\""), Value.Len(), Value.GetData());
+			return FString::Printf(TEXT("\"%s\""), *WriteToString<128>(Value));
 		}
 		else
 		{
@@ -96,7 +96,7 @@ namespace
 	{
 		if (Actual != Expected)
 		{
-			FAIL_CHECK(FString::Printf(TEXT("Expected '%s' to be %" "Iu" ", but it was %" "Iu" "."), What, Expected, Actual));
+			FAIL_CHECK(FString::Printf(TEXT("Expected '%s' to be %zu, but it was %zu."), What, Expected, Actual));
 			return false;
 		}
 		return true;
@@ -432,4 +432,4 @@ namespace
 	#define CHECK_NOT_EQUALS(What, X, Y)			TestNotEqual(What, X, Y);
 	#define CHECK_NOT_EQUALS_SENSITIVE(What, X, Y)	TestNotEqualSensitive(What, X, Y);
 
-}
\ No newline at end of file
+}
diff --git a/Engine/Source/Runtime/Core/Private/Misc/AutomationTest.cpp b/Engine/Source/Runtime/Core/Private/Misc/AutomationTest.cpp
--- a/Engine/Source/Runtime/Core/Private/Misc/AutomationTest.cpp
+++ b/Engine/Source/Runtime/Core/Private/Misc/AutomationTest.cpp
@@ -1608,7 +1608,7 @@ FString FAutomationTestBase::GetStringValueToDisplay(FUtf8StringView Value) cons
 {
 	if (Value.GetData())
 	{
-		return FString::Printf(TEXT("\"%.*hs\""), Value.Len(), Value.GetData());
+		return FString::Printf(TEXT("\"%s\""), *WriteToString<128>(Value));
 	}
 	else
 	{
