LowLevelTests: Added support for RaceDetector

diff --git a/Engine/Source/Developer/LowLevelTestsRunner/Private/TestRunner.cpp b/Engine/Source/Developer/LowLevelTestsRunner/Private/TestRunner.cpp
--- a/Engine/Source/Developer/LowLevelTestsRunner/Private/TestRunner.cpp
+++ b/Engine/Source/Developer/LowLevelTestsRunner/Private/TestRunner.cpp
@@ -10,6 +10,7 @@
 #include "Misc/StringBuilder.h"
 #include "Misc/Paths.h"
 #include "Modules/ModuleManager.h"
+#include "Sanitizer/RaceDetector.h"
 #include "String/Find.h"
 #include "String/LexFromString.h"
 #include "Containers/UnrealString.h"
@@ -287,6 +288,15 @@ int32 FTestRunner::RunCatchSession() const
 
 int RunTests(int32 ArgC, const ANSICHAR* ArgV[])
 {
+#if USING_INSTRUMENTATION
+	UE::Sanitizer::RaceDetector::Initialize();
+
+	ON_SCOPE_EXIT
+	{
+		UE::Sanitizer::RaceDetector::Shutdown();
+	};
+#endif
+
 	UE::LowLevelTests::FTestRunner TestRunner;
 
 	// Read command-line from file (if any). Some platforms do this earlier.
