Adressing CleanupLogs linker errors - remove it completely in favor of current use in TestRunner.

#rnx

diff --git a/Engine/Plugins/Experimental/StateGraph/Tests/StateGraphTests/Private/TestGroupEvents.cpp b/Engine/Plugins/Experimental/StateGraph/Tests/StateGraphTests/Private/TestGroupEvents.cpp
--- a/Engine/Plugins/Experimental/StateGraph/Tests/StateGraphTests/Private/TestGroupEvents.cpp
+++ b/Engine/Plugins/Experimental/StateGraph/Tests/StateGraphTests/Private/TestGroupEvents.cpp
@@ -5,6 +5,5 @@
 
 GROUP_AFTER_GLOBAL(Catch::DefaultGroup)
 {
-	CleanupLogs();
 	CleanupLocalization();
 }
diff --git a/Engine/Source/Developer/LowLevelTestsRunner/Public/TestCommon/Initialization.h b/Engine/Source/Developer/LowLevelTestsRunner/Public/TestCommon/Initialization.h
--- a/Engine/Source/Developer/LowLevelTestsRunner/Public/TestCommon/Initialization.h
+++ b/Engine/Source/Developer/LowLevelTestsRunner/Public/TestCommon/Initialization.h
@@ -18,6 +18,5 @@ void InitStats();
 
 void InitAll(bool bAllowLogging, bool bMultithreaded);
 
-void CleanupLogs();
 void CleanupLocalization();
 void CleanupAll();
\ No newline at end of file
diff --git a/Engine/Source/Programs/LowLevelTests/MathCoreTests/Private/TestGroupEvents.cpp b/Engine/Source/Programs/LowLevelTests/MathCoreTests/Private/TestGroupEvents.cpp
--- a/Engine/Source/Programs/LowLevelTests/MathCoreTests/Private/TestGroupEvents.cpp
+++ b/Engine/Source/Programs/LowLevelTests/MathCoreTests/Private/TestGroupEvents.cpp
@@ -10,6 +10,5 @@ GROUP_BEFORE_GLOBAL(Catch::DefaultGroup)
 
 GROUP_AFTER_GLOBAL(Catch::DefaultGroup)
 {
-	CleanupLogs();
 	CleanupLocalization();
 }
diff --git a/Engine/Source/Programs/WebTests/Private/WebLLTModule.cpp b/Engine/Source/Programs/WebTests/Private/WebLLTModule.cpp
--- a/Engine/Source/Programs/WebTests/Private/WebLLTModule.cpp
+++ b/Engine/Source/Programs/WebTests/Private/WebLLTModule.cpp
@@ -12,7 +12,6 @@
 
 GROUP_AFTER_GLOBAL(Catch::DefaultGroup)
 {
-	CleanupLogs();
 	CleanupLocalization();
 }
 
diff --git a/Engine/Source/Runtime/Online/Experimental/EventLoopTests/Tests/TestGroupEvents.cpp b/Engine/Source/Runtime/Online/Experimental/EventLoopTests/Tests/TestGroupEvents.cpp
--- a/Engine/Source/Runtime/Online/Experimental/EventLoopTests/Tests/TestGroupEvents.cpp
+++ b/Engine/Source/Runtime/Online/Experimental/EventLoopTests/Tests/TestGroupEvents.cpp
@@ -5,6 +5,5 @@
 
 GROUP_AFTER_GLOBAL(Catch::DefaultGroup)
 {
-	CleanupLogs();
 	CleanupLocalization();
 }
