Fix for wider range of low level tests issues

diff --git a/Engine/Source/Developer/LowLevelTestsRunner/Private/TestListeners/EventListener.cpp b/Engine/Source/Developer/LowLevelTestsRunner/Private/TestListeners/EventListener.cpp
--- a/Engine/Source/Developer/LowLevelTestsRunner/Private/TestListeners/EventListener.cpp
+++ b/Engine/Source/Developer/LowLevelTestsRunner/Private/TestListeners/EventListener.cpp
@@ -52,6 +52,20 @@ private:
 public:
 	using EventListenerBase::EventListenerBase;
 
+	void sectionStarting(const Catch::SectionInfo& SectionInfo) final
+	{
+		if (ITestRunner* TestRunner = ITestRunner::Get())
+		{
+			if (SectionInfo.name != CurrentTest)
+			{
+				UE_CLOG(TestRunner->IsDebugMode(), LogLowLevelTests, Display,
+					TEXT("running section \"%hs\" (line %" SIZE_T_FMT ")"),
+					SectionInfo.name.c_str(),
+					SIZE_T(SectionInfo.lineInfo.line));
+			}
+		}
+	}
+
 	void testCaseStarting(const Catch::TestCaseInfo& TestInfo) final
 	{
 		if (ITestRunner* TestRunner = ITestRunner::Get())
diff --git a/Engine/Source/Developer/LowLevelTestsRunner/Tests/SelfTests.cpp b/Engine/Source/Developer/LowLevelTestsRunner/Tests/SelfTests.cpp
--- a/Engine/Source/Developer/LowLevelTestsRunner/Tests/SelfTests.cpp
+++ b/Engine/Source/Developer/LowLevelTestsRunner/Tests/SelfTests.cpp
@@ -7,7 +7,7 @@
 #include <catch2/catch_active_test.hpp>
 #include <string>
 
-TEST_CASE("LowLevelTestsRunner::SelfTests::Skip")
+TEST_CASE("LowLevelTestsRunner::SelfTests::Skip", "[SelfTests]")
 {
 	SKIP("SKIP is expected to work all the time.");
 	FAIL("Test wasn't skipped.");
diff --git a/Engine/Source/Runtime/CoreUObject/Tests/ObjectPinningTest.cpp b/Engine/Source/Runtime/CoreUObject/Tests/ObjectPinningTest.cpp
--- a/Engine/Source/Runtime/CoreUObject/Tests/ObjectPinningTest.cpp
+++ b/Engine/Source/Runtime/CoreUObject/Tests/ObjectPinningTest.cpp
@@ -67,6 +67,9 @@ TEST_CASE("CoreUObject::ObjectPinning::Stress Tests", "[CoreUObject][ObjectPinni
 	// We need to be sure we've done the static GC initialization before we start doing a garbage collection.
 	FGCObject::StaticInit();
 
+	// Ignore oversubscription limit
+	LowLevelTasks::FScheduler::Get().GetOversubscriptionLimitReachedEvent().Clear();
+
 	UE_BENCHMARK(5, [] { ObjectPinningStressTest(1000, 1); } );
 	UE_BENCHMARK(5, [] { ObjectPinningStressTest(1000, 2); });
 	UE_BENCHMARK(5, [] { ObjectPinningStressTest(1000, 4); });
diff --git a/Engine/Source/Runtime/Engine/Private/Actor.cpp b/Engine/Source/Runtime/Engine/Private/Actor.cpp
--- a/Engine/Source/Runtime/Engine/Private/Actor.cpp
+++ b/Engine/Source/Runtime/Engine/Private/Actor.cpp
@@ -1281,7 +1281,7 @@ void AActor::PostMigrate(const FUObjectMigrationContext& MigrationContext)
 		const bool bKeepOwnership = (MigrationSendType == EObjectMigrationSendType::Loan);
 		if (bKeepOwnership)
 		{
-			RemoteObjects_ActorLoanedStates.AddAnnotation(this, { .bActorWasDestroyed = bActorIsBeingDestroyed, .bTickFunctionsRegistered = bTickFunctionsRegistered });
+			RemoteObjects_ActorLoanedStates.AddAnnotation(this, { .bActorWasDestroyed = (bool)bActorIsBeingDestroyed, .bTickFunctionsRegistered = (bool)bTickFunctionsRegistered });
 		}
 
 		if (!bKeepOwnership || !bHasOwnership)
