Only add --break if it's not already present, doubling up confuses Catch2

diff --git a/Engine/Source/Developer/LowLevelTestsRunner/Private/TestRunner.cpp b/Engine/Source/Developer/LowLevelTestsRunner/Private/TestRunner.cpp
--- a/Engine/Source/Developer/LowLevelTestsRunner/Private/TestRunner.cpp
+++ b/Engine/Source/Developer/LowLevelTestsRunner/Private/TestRunner.cpp
@@ -179,8 +179,12 @@ void FTestRunner::ParseCommandLine(TConstArrayView<const ANSICHAR*> Args)
 			CatchArgs.Add(Arg.GetData());
 		}
 	}
-	// Break in the debugger on failed assertions when attached.
-	CatchArgs.Add("--break");
+
+	if (!CatchArgs.Contains(ANSITEXTVIEW("--break")))
+	{
+		// Break in the debugger on failed assertions when attached.
+		CatchArgs.Add("--break");
+	}
 }
 
 void FTestRunner::SleepOnInit() const
