Foundation Tests Android issues
#rnx

diff --git a/Engine/Source/Developer/LowLevelTestsRunner/Private/Platform/Android/AndroidTestRunner.cpp b/Engine/Source/Developer/LowLevelTestsRunner/Private/Platform/Android/AndroidTestRunner.cpp
--- a/Engine/Source/Developer/LowLevelTestsRunner/Private/Platform/Android/AndroidTestRunner.cpp
+++ b/Engine/Source/Developer/LowLevelTestsRunner/Private/Platform/Android/AndroidTestRunner.cpp
@@ -35,14 +35,13 @@ extern "C"
 		jboolean isCopy;
 		const char* Path = env->GetStringUTFChars(jpath, &isCopy);
 
-		jsize jargsCount = env->GetArrayLength(jargs);
+		const jsize jargsCount = env->GetArrayLength(jargs);
 
-		const char* Args[jargsCount + 1];
+		TUniquePtr<const char* []> Args = MakeUnique<const char* []>(jargsCount + 1);
 		if (Path != nullptr)
 		{
 			Args[0] = AndroidPath = Path;
 		}
-		env->ReleaseStringUTFChars(jpath, Path);
 
 		for (int i = 0; i < jargsCount; i++) {
 			jboolean isArgCopy;
@@ -57,7 +56,11 @@ extern "C"
 			env->ReleaseStringUTFChars(jarg, Arg);
 		}
 
-		return LowLevelTestsMain(jargsCount + 1, Args);
+		int MainReturn = LowLevelTestsMain(jargsCount + 1, Args.Get());
+
+		env->ReleaseStringUTFChars(jpath, Path);
+
+		return MainReturn;
 	}
 
 	JNIEXPORT jint JNI_OnLoad(JavaVM* InJavaVM, void* InReserved)
diff --git a/Engine/Source/Programs/UnrealBuildTool/Configuration/TestTargetRules.cs b/Engine/Source/Programs/UnrealBuildTool/Configuration/TestTargetRules.cs
--- a/Engine/Source/Programs/UnrealBuildTool/Configuration/TestTargetRules.cs
+++ b/Engine/Source/Programs/UnrealBuildTool/Configuration/TestTargetRules.cs
@@ -351,6 +351,7 @@ namespace UnrealBuildTool
 				string IntermediatePath = Path.Combine(Unreal.EngineDirectory.FullName, "Intermediate", "Android");
 				FileReference versionScriptFile = new FileReference(Path.Combine(IntermediatePath, $"LLTWorkaroundScript-{Name}.ldscript"));
 				string SymbolsLinkScriptContent = @"{
+					global: _Znwm*; _Znam*; _ZdlPv*; _ZdaPv*;
 					global: LowLevelTestsMain;
 					global: JNI_OnLoad;
 					global: Java_com_epicgames_unreal_tests_TestActivity_runTests;
diff --git a/Engine/Source/Programs/UnrealBuildTool/Platform/Android/AndroidToolChain.cs b/Engine/Source/Programs/UnrealBuildTool/Platform/Android/AndroidToolChain.cs
--- a/Engine/Source/Programs/UnrealBuildTool/Platform/Android/AndroidToolChain.cs
+++ b/Engine/Source/Programs/UnrealBuildTool/Platform/Android/AndroidToolChain.cs
@@ -788,6 +788,7 @@ namespace UnrealBuildTool
 			// make sure the DT_SONAME field is set properly (or we can a warning toast at startup on new Android)
 			Result += " -Wl,-soname,libUnreal.so";
 
+			if (!LinkEnvironment.AdditionalArguments.Contains("--version-script"))
 			{
 				string VersionScriptFile = GetVersionScriptFilename(LinkEnvironment);
 				using (StreamWriter Writer = File.CreateText(VersionScriptFile))
@@ -797,6 +798,10 @@ namespace UnrealBuildTool
 				}
 				Result += " -Wl,--version-script=\"" + VersionScriptFile + "\"";
 			}
+			else
+			{
+				Logger.LogInformation("Linker version script already passed through linker environment additional arguments.");
+			}
 
 			Result += " -Wl,--build-id=sha1";               // add build-id to make debugging easier
 
