LLT Android permissions update

#rnx

diff --git a/Engine/Source/Developer/LowLevelTestsRunner/Android/Java/src/com/epicgames/unreal/tests/TestActivity.java b/Engine/Source/Developer/LowLevelTestsRunner/Android/Java/src/com/epicgames/unreal/tests/TestActivity.java
--- a/Engine/Source/Developer/LowLevelTestsRunner/Android/Java/src/com/epicgames/unreal/tests/TestActivity.java
+++ b/Engine/Source/Developer/LowLevelTestsRunner/Android/Java/src/com/epicgames/unreal/tests/TestActivity.java
@@ -12,15 +12,32 @@ import java.io.IOException;
 import java.util.Arrays;
 import java.util.ArrayList;
 
+import androidx.core.app.ActivityCompat;
+
 import android.content.pm.PackageManager;
+import android.Manifest.permission;
 
 public class TestActivity extends Activity {
 	private int testsReturnCode = -1;
-	
+
 	public TestActivity() {
 	}
 
 	public String readCmdLineFrom(String path) {
+		// request permissions for read and write to external storage if needed
+		int permission = ActivityCompat.checkSelfPermission(this, android.Manifest.permission.WRITE_EXTERNAL_STORAGE);
+		String[] externalPermissions = {
+			android.Manifest.permission.READ_EXTERNAL_STORAGE,
+			android.Manifest.permission.WRITE_EXTERNAL_STORAGE
+		};
+		if (permission != PackageManager.PERMISSION_GRANTED) {
+			ActivityCompat.requestPermissions(
+				this,
+				externalPermissions,
+				1
+			);
+		}
+
 		try {
 			File cmdLineFile = new File(path);
 			FileInputStream cmdLineInputStream = new FileInputStream(cmdLineFile);
