[RaceDetector]
- Add race detector support in low level tests

diff --git a/Engine/Source/Developer/LowLevelTestsRunner/Private/TestCommon/Initialization.cpp b/Engine/Source/Developer/LowLevelTestsRunner/Private/TestCommon/Initialization.cpp
--- a/Engine/Source/Developer/LowLevelTestsRunner/Private/TestCommon/Initialization.cpp
+++ b/Engine/Source/Developer/LowLevelTestsRunner/Private/TestCommon/Initialization.cpp
@@ -11,7 +11,7 @@
 #include "Misc/CommandLine.h"
 #include "Misc/DelayedAutoRegister.h"
 #include "Misc/Paths.h"
-
+#include "Sanitizer/RaceDetector.h"
 
 static IPlatformFile* DefaultPlatformFile;
 
@@ -105,6 +105,10 @@ void SetProjectNameAndDirectory()
 
 void InitAll(bool bAllowLogging, bool bMultithreaded)
 {
+#if USING_INSTRUMENTATION
+	UE::Sanitizer::RaceDetector::Initialize();
+	UE::Sanitizer::RaceDetector::ToggleRaceDetection(true);
+#endif
 	SaveDefaultPlatformFile();
 	UsePlatformFileStubIfRequired();
 	InitAllThreadPools(bMultithreaded);
@@ -145,4 +149,8 @@ void CleanupAll()
 	CleanupAllThreadPools();
 	CleanupTaskGraph();
 	CleanupLocalization();
+
+#if USING_INSTRUMENTATION
+	UE::Sanitizer::RaceDetector::Shutdown();
+#endif
 }
\ No newline at end of file
