LLT issues - Android use correct path for command line file

#rnx

diff --git a/Engine/Source/Developer/LowLevelTestsRunner/Android/Java/src/com/epicgames/unreal/tests/TestActivity.java b/Engine/Source/Developer/LowLevelTestsRunner/Android/Java/src/com/epicgames/unreal/tests/TestActivity.java
--- a/Engine/Source/Developer/LowLevelTestsRunner/Android/Java/src/com/epicgames/unreal/tests/TestActivity.java
+++ b/Engine/Source/Developer/LowLevelTestsRunner/Android/Java/src/com/epicgames/unreal/tests/TestActivity.java
@@ -23,7 +23,7 @@ public class TestActivity extends Activity {
 	public TestActivity() {
 	}
 
-	public String readCmdLineFrom(String path) {
+	private String readCmdLineFrom(String path) {
 		// request permissions for read and write to external storage if needed
 		int permission = ActivityCompat.checkSelfPermission(this, android.Manifest.permission.WRITE_EXTERNAL_STORAGE);
 		String[] externalPermissions = {
@@ -73,7 +73,9 @@ public class TestActivity extends Activity {
 		{
 			String[] PackageNameParts = PackageName.split("\\.");
 			String PackageNameLastPart = PackageNameParts[PackageNameParts.length - 1];
-			String cmdLinePath = "/sdcard/UnrealGame/" + PackageNameLastPart + "/UECommandLine.txt";
+
+			String GamePath = getExternalFilesDir(null).getAbsolutePath() + "/UnrealGame/" + PackageNameLastPart;
+			String cmdLinePath = GamePath + "/UECommandLine.txt";
 
 			String TestCommandLine = readCmdLineFrom(cmdLinePath);
 			String AbsolutePath = getFilesDir().getAbsolutePath() + "/";
diff --git a/Engine/Source/Programs/AutomationTool/LowLevelTests/RunLowLevelTests.cs b/Engine/Source/Programs/AutomationTool/LowLevelTests/RunLowLevelTests.cs
--- a/Engine/Source/Programs/AutomationTool/LowLevelTests/RunLowLevelTests.cs
+++ b/Engine/Source/Programs/AutomationTool/LowLevelTests/RunLowLevelTests.cs
@@ -840,7 +840,7 @@ namespace LowLevelTests
 		public IBuild CreateBuild(UnrealTargetPlatform InPlatform, UnrealTargetConfiguration InConfiguration, string InTestApp, string InBuildPath, bool bSkipStage)
 		{
 			string ExecutablePath = ILowLevelTestsBuildFactory.GetExecutable(InPlatform, InConfiguration, InTestApp, InBuildPath, GetExecutableRegex(InPlatform));
-			return new AndroidBuild(InConfiguration, $"com.epicgames.{InTestApp}", Path.Combine(InBuildPath, ExecutablePath), new Dictionary<string, string>(), BuildFlags.Packaged | BuildFlags.CanReplaceCommandLine, false, false, false);
+			return new AndroidBuild(InConfiguration, $"com.epicgames.{InTestApp}", Path.Combine(InBuildPath, ExecutablePath), new Dictionary<string, string>(), BuildFlags.Packaged | BuildFlags.CanReplaceCommandLine, false, true, false);
 		}
 
 		public string GetExecutableRegex(UnrealTargetPlatform InPlatform)
diff --git a/Engine/Source/Programs/UnrealBuildTool/Modes/TestMode.cs b/Engine/Source/Programs/UnrealBuildTool/Modes/TestMode.cs
--- a/Engine/Source/Programs/UnrealBuildTool/Modes/TestMode.cs
+++ b/Engine/Source/Programs/UnrealBuildTool/Modes/TestMode.cs
@@ -7,14 +7,11 @@ using System.Linq;
 using System.Reflection;
 using System.Threading.Tasks;
 using EpicGames.Core;
-using Microsoft.Build.Execution;
 using Microsoft.Extensions.Logging;
 using UnrealBuildBase;
 
 namespace UnrealBuildTool
 {
-	using IncludeChainMap = Dictionary<UEBuildModule, List<string>>;
-
 	/// <summary>
 	/// Builds low level tests on one or more targets.
 	/// </summary>
