LLT Android - Remove runtime permission request and add external storage read and write permissions to manifest

#rnx

diff --git a/Engine/Source/Developer/LowLevelTestsRunner/Android/Java/src/com/epicgames/unreal/tests/TestActivity.java b/Engine/Source/Developer/LowLevelTestsRunner/Android/Java/src/com/epicgames/unreal/tests/TestActivity.java
--- a/Engine/Source/Developer/LowLevelTestsRunner/Android/Java/src/com/epicgames/unreal/tests/TestActivity.java
+++ b/Engine/Source/Developer/LowLevelTestsRunner/Android/Java/src/com/epicgames/unreal/tests/TestActivity.java
@@ -24,20 +24,6 @@ public class TestActivity extends Activity {
 	}
 
 	private String readCmdLineFrom(String path) {
-		// request permissions for read and write to external storage if needed
-		int permission = ActivityCompat.checkSelfPermission(this, android.Manifest.permission.WRITE_EXTERNAL_STORAGE);
-		String[] externalPermissions = {
-			android.Manifest.permission.READ_EXTERNAL_STORAGE,
-			android.Manifest.permission.WRITE_EXTERNAL_STORAGE
-		};
-		if (permission != PackageManager.PERMISSION_GRANTED) {
-			ActivityCompat.requestPermissions(
-				this,
-				externalPermissions,
-				1
-			);
-		}
-
 		try {
 			File cmdLineFile = new File(path);
 			FileInputStream cmdLineInputStream = new FileInputStream(cmdLineFile);
diff --git a/Engine/Source/Developer/LowLevelTestsRunner/LowLevelTestsAndroid_UPL.xml b/Engine/Source/Developer/LowLevelTestsRunner/LowLevelTestsAndroid_UPL.xml
--- a/Engine/Source/Developer/LowLevelTestsRunner/LowLevelTestsAndroid_UPL.xml
+++ b/Engine/Source/Developer/LowLevelTestsRunner/LowLevelTestsAndroid_UPL.xml
@@ -7,6 +7,14 @@
 		<log text="Low level tests UPL script"/>
 	</init>
 	<androidManifestUpdates>
+		<setElement result="PermissionRead" value="uses-permission" />
+		<addAttribute tag="$PermissionRead" name="android:name" value="android.permission.READ_EXTERNAL_STORAGE" />
+		<addElement tag="manifest" name="PermissionRead" />
+
+		<setElement result="PermissionWrite" value="uses-permission" />
+		<addAttribute tag="$PermissionWrite" name="android:name" value="android.permission.Write_EXTERNAL_STORAGE" />
+		<addElement tag="manifest" name="PermissionWrite" />
+
 		<removeElement tag="application"/>
 		<setElement result="TestApplication" value="application" />
 		<addAttribute tag="$TestApplication" name="android:name" value="com.epicgames.unreal.tests.TestApplication" />
