Fix CIS issues in non unity builds of iOS/macOS targets.

This is fallout of a previous change to remove sdk headers of platform files, which means files in other modules need to manually include them. I missed some files used by a few targets and it wasn't caught by the preflight in the original commit

#Synced-CL 49794922

diff --git a/Engine/Plugins/Runtime/AR/AppleAR/AppleARKit/Source/AppleARKit/Private/AppleARKitTrackable.cpp b/Engine/Plugins/Runtime/AR/AppleAR/AppleARKit/Source/AppleARKit/Private/AppleARKitTrackable.cpp
--- a/Engine/Plugins/Runtime/AR/AppleAR/AppleARKit/Source/AppleARKit/Private/AppleARKitTrackable.cpp
+++ b/Engine/Plugins/Runtime/AR/AppleAR/AppleARKit/Source/AppleARKit/Private/AppleARKitTrackable.cpp
@@ -3,6 +3,10 @@
 #include "AppleARKitTrackable.h"
 #include "AppleARKitTextures.h"
 
+#if PLATFORM_IOS
+#import <Metal/MTLTexture.h>
+#endif
+
 UAppleARKitEnvironmentCaptureProbe::UAppleARKitEnvironmentCaptureProbe()
 	: Super()
 	, ARKitEnvironmentTexture(nullptr)
diff --git a/Engine/Plugins/Runtime/AR/AppleAR/AppleARKit/Source/AppleARKit/Private/AppleARKitTrackable.h b/Engine/Plugins/Runtime/AR/AppleAR/AppleARKit/Source/AppleARKit/Private/AppleARKitTrackable.h
--- a/Engine/Plugins/Runtime/AR/AppleAR/AppleARKit/Source/AppleARKit/Private/AppleARKitTrackable.h
+++ b/Engine/Plugins/Runtime/AR/AppleAR/AppleARKit/Source/AppleARKit/Private/AppleARKitTrackable.h
@@ -6,6 +6,12 @@
 #include "ARTrackable.h"
 #include "AppleARKitTrackable.generated.h"
 
+#if PLATFORM_IOS
+
+@protocol MTLTexture;
+
+#endif
+
 class UAppleARKitEnvironmentCaptureProbeTexture;
 
 UCLASS(NotBlueprintType)
diff --git a/Engine/Source/Developer/LowLevelTestsRunner/Private/Platform/Apple/AppleTestRunnerHelper.cpp b/Engine/Source/Developer/LowLevelTestsRunner/Private/Platform/Apple/AppleTestRunnerHelper.cpp
--- a/Engine/Source/Developer/LowLevelTestsRunner/Private/Platform/Apple/AppleTestRunnerHelper.cpp
+++ b/Engine/Source/Developer/LowLevelTestsRunner/Private/Platform/Apple/AppleTestRunnerHelper.cpp
@@ -3,6 +3,16 @@
 
 #if PLATFORM_IOS || PLATFORM_MAC
 
+#pragma push_macro("FVector")
+#define FVector FVectorWorkaround
+
+#import <CoreFoundation/CFRunLoop.h>
+#import <Foundation/Foundation.h>
+
+#undef FVector
+
+#pragma pop_macro("FVector")
+
 #include "TestRunner.h"
 
 @implementation AppleTestsRunnerHelper
diff --git a/Engine/Source/Developer/LowLevelTestsRunner/Private/Platform/IOS/IOSTestRunner.cpp b/Engine/Source/Developer/LowLevelTestsRunner/Private/Platform/IOS/IOSTestRunner.cpp
--- a/Engine/Source/Developer/LowLevelTestsRunner/Private/Platform/IOS/IOSTestRunner.cpp
+++ b/Engine/Source/Developer/LowLevelTestsRunner/Private/Platform/IOS/IOSTestRunner.cpp
@@ -5,6 +5,7 @@
 #include "TestRunner.h"
 #include "Platform/Apple/AppleTestRunnerHelper.h"
 #include <mach-o/dyld.h>
+#import <CoreFoundation/CFRunLoop.h>
 
 const char* GetCacheDirectory()
 {
diff --git a/Engine/Source/Runtime/Core/Private/Apple/AppleMemoryObject.cpp b/Engine/Source/Runtime/Core/Private/Apple/AppleMemoryObject.cpp
--- a/Engine/Source/Runtime/Core/Private/Apple/AppleMemoryObject.cpp
+++ b/Engine/Source/Runtime/Core/Private/Apple/AppleMemoryObject.cpp
@@ -30,6 +30,10 @@ extern "C" int proc_pid_rusage(int pid, int flavor, rusage_info_t *buffer)
 	__OSX_AVAILABLE_STARTING(__MAC_10_9, __IPHONE_7_0);
 #endif
 
+#include "HAL/UnrealMemory.h"
+#include "HAL/PlatformMath.h"
+#include "Templates/AlignmentTemplates.h"
+
 NS_ASSUME_NONNULL_BEGIN
 
 /** 
diff --git a/Engine/Source/Runtime/Core/Tests/HAL/PreprocessorHelperTests.cpp b/Engine/Source/Runtime/Core/Tests/HAL/PreprocessorHelperTests.cpp
--- a/Engine/Source/Runtime/Core/Tests/HAL/PreprocessorHelperTests.cpp
+++ b/Engine/Source/Runtime/Core/Tests/HAL/PreprocessorHelperTests.cpp
@@ -2,6 +2,7 @@
 
 #include "HAL/PreprocessorHelpers.h"
 
+#include "Misc/CString.h"
 #include "Tests/TestHarnessAdapter.h"
 
 #if WITH_TESTS
diff --git a/Engine/Source/Runtime/Online/BackgroundHTTP/Private/IOS/ApplePlatformBackgroundHttpRequest.cpp b/Engine/Source/Runtime/Online/BackgroundHTTP/Private/IOS/ApplePlatformBackgroundHttpRequest.cpp
--- a/Engine/Source/Runtime/Online/BackgroundHTTP/Private/IOS/ApplePlatformBackgroundHttpRequest.cpp
+++ b/Engine/Source/Runtime/Online/BackgroundHTTP/Private/IOS/ApplePlatformBackgroundHttpRequest.cpp
@@ -7,6 +7,8 @@
 #include "HAL/PlatformFile.h"
 #include "HAL/PlatformFileManager.h"
 
+#import <Foundation/NSURLSession.h>
+
 FApplePlatformBackgroundHttpRequest::FApplePlatformBackgroundHttpRequest()
 	: DownloadId(FBackgroundURLSessionHandler::InvalidDownloadId)
 {
