[LowLevelTestRunner]
* Fixed compile error on mac

diff --git a/Engine/Source/Developer/LowLevelTestsRunner/LowLevelTestsRunner.Build.cs b/Engine/Source/Developer/LowLevelTestsRunner/LowLevelTestsRunner.Build.cs
--- a/Engine/Source/Developer/LowLevelTestsRunner/LowLevelTestsRunner.Build.cs
+++ b/Engine/Source/Developer/LowLevelTestsRunner/LowLevelTestsRunner.Build.cs
@@ -8,6 +8,8 @@ namespace UnrealBuildTool.Rules
 	{
 		public LowLevelTestsRunner(ReadOnlyTargetRules Target) : base(Target)
 		{
+			bRequiresPlatformSDK = true;
+			
 			PCHUsage = PCHUsageMode.NoPCHs;
 			PrecompileForTargets = PrecompileTargetsType.None;
 			bUseUnity = false;
diff --git a/Engine/Source/Developer/LowLevelTestsRunner/Private/Platform/Mac/MacTestRunner.cpp b/Engine/Source/Developer/LowLevelTestsRunner/Private/Platform/Mac/MacTestRunner.cpp
--- a/Engine/Source/Developer/LowLevelTestsRunner/Private/Platform/Mac/MacTestRunner.cpp
+++ b/Engine/Source/Developer/LowLevelTestsRunner/Private/Platform/Mac/MacTestRunner.cpp
@@ -3,6 +3,7 @@
 #if PLATFORM_MAC
 
 #include "TestRunner.h"
+#include "Mac/MacSystemIncludes.h"
 #include "Platform/Apple/AppleTestRunnerHelper.h"
 #include <mach-o/dyld.h>
 
