Force GameThread scope before the TestRunner is created for Low Level Tests.

#Synced-CL 50506220

diff --git a/Engine/Source/Developer/LowLevelTestsRunner/Private/TestRunner.cpp b/Engine/Source/Developer/LowLevelTestsRunner/Private/TestRunner.cpp
--- a/Engine/Source/Developer/LowLevelTestsRunner/Private/TestRunner.cpp
+++ b/Engine/Source/Developer/LowLevelTestsRunner/Private/TestRunner.cpp
@@ -302,6 +302,8 @@ int RunTests(int32 ArgC, const ANSICHAR* ArgV[])
 	};
 #endif
 
+	FTaskTagScope Scope(ETaskTag::EGameThread);
+	
 	UE::LowLevelTests::FTestRunner TestRunner;
 
 	// Read command-line from file (if any). Some platforms do this earlier.
@@ -339,4 +341,4 @@ int RunTests(int32 ArgC, const ANSICHAR* ArgV[])
 
 	int CatchReturn = TestRunner.RunCatchSession();
 	return CatchReturn;
-}
\ No newline at end of file
+}
