Fix:  Back out changes from CL 50564093

#Synced-CL 50506220

diff --git a/Engine/Source/Developer/LowLevelTestsRunner/Private/TestRunner.cpp b/Engine/Source/Developer/LowLevelTestsRunner/Private/TestRunner.cpp
--- a/Engine/Source/Developer/LowLevelTestsRunner/Private/TestRunner.cpp
+++ b/Engine/Source/Developer/LowLevelTestsRunner/Private/TestRunner.cpp
@@ -301,8 +301,6 @@ int RunTests(int32 ArgC, const ANSICHAR* ArgV[])
 		UE::Sanitizer::RaceDetector::Shutdown();
 	};
 #endif
-
-	FTaskTagScope Scope(ETaskTag::EGameThread);
 	
 	UE::LowLevelTests::FTestRunner TestRunner;
 
