On MacOS, force GameThread scope before the TestRunner is created for Low Level Tests

#Synced-CL 50506220

diff --git a/Engine/Source/Developer/LowLevelTestsRunner/Private/Platform/Apple/AppleTestRunnerHelper.cpp b/Engine/Source/Developer/LowLevelTestsRunner/Private/Platform/Apple/AppleTestRunnerHelper.cpp
--- a/Engine/Source/Developer/LowLevelTestsRunner/Private/Platform/Apple/AppleTestRunnerHelper.cpp
+++ b/Engine/Source/Developer/LowLevelTestsRunner/Private/Platform/Apple/AppleTestRunnerHelper.cpp
@@ -48,6 +48,8 @@ int _Result;
 
 - (void)runTestsThread:(id)Arg
 {
+	FTaskTagScope Scope(ETaskTag::EGameThread);
+	
 	_Result = RunTests(Argc, Argv);
 	
 	CFRunLoopStop(CFRunLoopGetMain());
