LTTs: Remove duplicate tag scope Mac (fixes failures on Mac tests)

#rnx

#Synced-CL 0

diff --git a/Engine/Source/Developer/LowLevelTestsRunner/Private/Platform/Apple/AppleTestRunnerHelper.cpp b/Engine/Source/Developer/LowLevelTestsRunner/Private/Platform/Apple/AppleTestRunnerHelper.cpp
--- a/Engine/Source/Developer/LowLevelTestsRunner/Private/Platform/Apple/AppleTestRunnerHelper.cpp
+++ b/Engine/Source/Developer/LowLevelTestsRunner/Private/Platform/Apple/AppleTestRunnerHelper.cpp
@@ -48,8 +48,6 @@ int _Result;
 
 - (void)runTestsThread:(id)Arg
 {
-	FTaskTagScope Scope(ETaskTag::EGameThread);
-	
 	_Result = RunTests(Argc, Argv);
 	
 	CFRunLoopStop(CFRunLoopGetMain());
