LowLevelTests: Correctly set ETaskTag::GameThread for the lifetime of the test runner

#Synced-CL 50253651

diff --git a/Engine/Source/Developer/LowLevelTestsRunner/Private/TestCommon/CoreUtilities.cpp b/Engine/Source/Developer/LowLevelTestsRunner/Private/TestCommon/CoreUtilities.cpp
--- a/Engine/Source/Developer/LowLevelTestsRunner/Private/TestCommon/CoreUtilities.cpp
+++ b/Engine/Source/Developer/LowLevelTestsRunner/Private/TestCommon/CoreUtilities.cpp
@@ -89,7 +89,6 @@ void CleanupAllThreadPools()
 
 void InitAllThreadPools(bool MultiThreaded)
 {
-	FTaskTagScope Scope(ETaskTag::EGameThread);
 	InitIOThreadPool(MultiThreaded);
 	InitThreadPool(MultiThreaded);
 	InitBackgroundPriorityThreadPool(MultiThreaded);
@@ -97,7 +96,6 @@ void InitAllThreadPools(bool MultiThreaded)
 
 void InitTaskGraph(bool MultiThreaded, ENamedThreads::Type ThreadToAttach)
 {
-	FTaskTagScope Scope(ETaskTag::EGameThread);
 	FTaskGraphInterface::Startup(MultiThreaded ? FPlatformMisc::NumberOfWorkerThreadsToSpawn() : 1);
 	FTaskGraphInterface::Get().AttachToThread(ThreadToAttach);
 
diff --git a/Engine/Source/Developer/LowLevelTestsRunner/Private/TestRunner.cpp b/Engine/Source/Developer/LowLevelTestsRunner/Private/TestRunner.cpp
--- a/Engine/Source/Developer/LowLevelTestsRunner/Private/TestRunner.cpp
+++ b/Engine/Source/Developer/LowLevelTestsRunner/Private/TestRunner.cpp
@@ -301,7 +301,9 @@ int RunTests(int32 ArgC, const ANSICHAR* ArgV[])
 		UE::Sanitizer::RaceDetector::Shutdown();
 	};
 #endif
-	
+
+	FTaskTagScope Scope(ETaskTag::EGameThread);
+
 	UE::LowLevelTests::FTestRunner TestRunner;
 
 	// Read command-line from file (if any). Some platforms do this earlier.
