LLT TestData json reporting

#rnx
#Synced-CL 0

#Synced-CL 0

#Synced-CL 0

diff --git a/Engine/Build/LowLevelTests.xml b/Engine/Build/LowLevelTests.xml
--- a/Engine/Build/LowLevelTests.xml
+++ b/Engine/Build/LowLevelTests.xml
@@ -2,6 +2,12 @@
 <BuildGraph xmlns="http://www.epicgames.com/BuildGraph" xmlns:xsi="http://www.w3.org/2001/XMLSchema-instance" xsi:schemaLocation="http://www.epicgames.com/BuildGraph ./Graph/Schema.xsd" >
 	<EnvVar Name="UE_HORDE_JOBID"/>
 
+	<Option Name="PreflightChange" DefaultValue="" Description="Preflight CL number if preflight, empty otherwise"/>
+	<Property Name="IsPreflight" Value="false"/>
+	<Property Name="IsPreflight" Value="true" If="'$(PreflightChange)' != ''"/>
+	<Property Name="PreflightSuffix" Value=""/>
+	<Property Name="PreflightSuffix" Value="-PF-$(PreflightChange)-$(UE_HORDE_JOBID)" If="$(IsPreflight)"/>
+
 	<Property Name="TestNames" Value=""/>
 
 	<Property Name="AllTestablePlatforms" Value="Win64;Mac;MacArm64;Linux;LinuxArm64;Android;IOS"/>
@@ -242,8 +248,10 @@
 						<Property Name="TestArgsExt" Value="$(TestArgsExt) -adbkeys=$(AdbKeysPath) -captureoutput"/>
 					</Do>
 
+					<Property Name="TestDataKey" Value="$(ShortName) Tests $(Platform)$(PreflightSuffix)"/>
+
 					<Command Name="RunLowLevelTests"
-							Arguments="-test=LowLevelTests -testapp=$(TargetName)$(TestModeSuffix) -platform=$(Platform) -configuration=$(Configuration) $(TestArgsExt) -VeryVerbose" If="!$(DryRun) and '$(RunUnsupported)' == ''"/>
+							Arguments="-test=LowLevelTests -testapp=$(TargetName)$(TestModeSuffix) -platform=$(Platform) -configuration=$(Configuration) $(TestArgsExt) -HordeTestDataKey=&quot;$(TestDataKey)&quot; -VeryVerbose" If="!$(DryRun) and '$(RunUnsupported)' == ''"/>
 
 					<Do If="'$(AppRunsContainerized)' == 'True'">
 						<Property Name="TestTargetLowercase" Value=""/>
diff --git a/Engine/Build/LowLevelTests/ChasoUserDataPT.xml b/Engine/Build/LowLevelTests/ChasoUserDataPT.xml
--- a/Engine/Build/LowLevelTests/ChasoUserDataPT.xml
+++ b/Engine/Build/LowLevelTests/ChasoUserDataPT.xml
@@ -3,6 +3,6 @@
   <Option Name="RunChasoUserDataPTTests" DefaultValue="" Description="Run Chaos User Data PT Tests" />
   <Property Name="TestNames" Value="$(TestNames);ChasoUserDataPT" />
   <Extend Name="RunAllTests">
-    <Expand Name="DeployAndTest" Platform="Win64" TestName="ChasoUserDataPT" ShortName="Chaos User Data PT" TargetName="ChaosUserDataPTTests" BinaryRelativePath="Engine\Binaries" ReportType="console" />
+    <Expand Name="DeployAndTest" Platform="Win64" TestName="ChasoUserDataPT" ShortName="Chaos User Data PT" TargetName="ChaosUserDataPTTests" BinaryRelativePath="Engine\Binaries" ReportType="testdata" />
   </Extend>
 </BuildGraph>
\ No newline at end of file
diff --git a/Engine/Build/LowLevelTests/EventLoop.xml b/Engine/Build/LowLevelTests/EventLoop.xml
--- a/Engine/Build/LowLevelTests/EventLoop.xml
+++ b/Engine/Build/LowLevelTests/EventLoop.xml
@@ -3,6 +3,6 @@
   <Option Name="RunEventLoopTests" DefaultValue="" Description="Run EventLoop Tests" />
   <Property Name="TestNames" Value="$(TestNames);EventLoop" />
   <Extend Name="RunAllTests">
-    <Expand Name="DeployAndTest" Platform="Win64" TestName="EventLoop" ShortName="EventLoop" TargetName="EventLoopUnitTests" BinaryRelativePath="Engine\Binaries" ReportType="console" />
+    <Expand Name="DeployAndTest" Platform="Win64" TestName="EventLoop" ShortName="EventLoop" TargetName="EventLoopUnitTests" BinaryRelativePath="Engine\Binaries" ReportType="testdata" />
   </Extend>
 </BuildGraph>
\ No newline at end of file
diff --git a/Engine/Build/LowLevelTests/Foundation.xml b/Engine/Build/LowLevelTests/Foundation.xml
--- a/Engine/Build/LowLevelTests/Foundation.xml
+++ b/Engine/Build/LowLevelTests/Foundation.xml
@@ -3,10 +3,10 @@
   <Option Name="RunFoundationTests" DefaultValue="" Description="Run Foundation Tests" />
   <Property Name="TestNames" Value="$(TestNames);Foundation" />
   <Extend Name="RunAllTests">
-    <Expand Name="DeployAndTest" Platform="Win64" TestName="Foundation" ShortName="Foundation" TargetName="FoundationTests" BinaryRelativePath="Engine\Binaries" ReportType="xml" ExtraCompilationArgs="-allmodules" />
-    <Expand Name="DeployAndTest" Platform="Linux" TestName="Foundation" ShortName="Foundation" TargetName="FoundationTests" BinaryRelativePath="Engine\Binaries" ReportType="xml" Tags="~[.]~[Slow]" ExtraCompilationArgs="-allmodules" />
-    <Expand Name="DeployAndTest" Platform="Mac" TestName="Foundation" ShortName="Foundation" TargetName="FoundationTests" BinaryRelativePath="Engine\Binaries" ReportType="xml" ExtraCompilationArgs="-allmodules" />
-    <Expand Name="DeployAndTest" Platform="Android" TestName="Foundation" ShortName="Foundation" TargetName="FoundationTests" BinaryRelativePath="Engine\Binaries" ReportType="xml" Tags="~[Perf]~[Slow]~[AndroidSkip]" ExtraCompilationArgs="-allmodules -architectures=arm64" />
-    <Expand Name="DeployAndTest" Platform="IOS" TestName="Foundation" ShortName="Foundation" TargetName="FoundationTests" BinaryRelativePath="Engine\Binaries" ReportType="xml" ExtraCompilationArgs="-allmodules" RunUnsupported="True" />
+    <Expand Name="DeployAndTest" Platform="Win64" TestName="Foundation" ShortName="Foundation" TargetName="FoundationTests" BinaryRelativePath="Engine\Binaries" ReportType="testdata" ExtraCompilationArgs="-allmodules" />
+    <Expand Name="DeployAndTest" Platform="Linux" TestName="Foundation" ShortName="Foundation" TargetName="FoundationTests" BinaryRelativePath="Engine\Binaries" ReportType="testdata" Tags="~[.]~[Slow]" ExtraCompilationArgs="-allmodules" />
+    <Expand Name="DeployAndTest" Platform="Mac" TestName="Foundation" ShortName="Foundation" TargetName="FoundationTests" BinaryRelativePath="Engine\Binaries" ReportType="testdata" ExtraCompilationArgs="-allmodules" />
+    <Expand Name="DeployAndTest" Platform="Android" TestName="Foundation" ShortName="Foundation" TargetName="FoundationTests" BinaryRelativePath="Engine\Binaries" ReportType="console" Tags="~[Perf]~[Slow]~[AndroidSkip]" ExtraCompilationArgs="-allmodules -architectures=arm64" />
+    <Expand Name="DeployAndTest" Platform="IOS" TestName="Foundation" ShortName="Foundation" TargetName="FoundationTests" BinaryRelativePath="Engine\Binaries" ReportType="testdata" ExtraCompilationArgs="-allmodules" RunUnsupported="True" />
   </Extend>
 </BuildGraph>
\ No newline at end of file
diff --git a/Engine/Build/LowLevelTests/GameplayGraph.xml b/Engine/Build/LowLevelTests/GameplayGraph.xml
--- a/Engine/Build/LowLevelTests/GameplayGraph.xml
+++ b/Engine/Build/LowLevelTests/GameplayGraph.xml
@@ -3,6 +3,6 @@
   <Option Name="RunGameplayGraphTests" DefaultValue="" Description="Run GameplayGraph Tests" />
   <Property Name="TestNames" Value="$(TestNames);GameplayGraph" />
   <Extend Name="RunAllTests">
-    <Expand Name="DeployAndTest" Platform="Win64" TestName="GameplayGraph" ShortName="GameplayGraph" TargetName="GameplayGraphTests" BinaryRelativePath="Engine\Binaries" ReportType="console" />
+    <Expand Name="DeployAndTest" Platform="Win64" TestName="GameplayGraph" ShortName="GameplayGraph" TargetName="GameplayGraphTests" BinaryRelativePath="Engine\Binaries" ReportType="testdata" />
   </Extend>
 </BuildGraph>
\ No newline at end of file
diff --git a/Engine/Build/LowLevelTests/GlobalConfigurationData.xml b/Engine/Build/LowLevelTests/GlobalConfigurationData.xml
--- a/Engine/Build/LowLevelTests/GlobalConfigurationData.xml
+++ b/Engine/Build/LowLevelTests/GlobalConfigurationData.xml
@@ -3,6 +3,6 @@
   <Option Name="RunGlobalConfigurationDataTests" DefaultValue="" Description="Run GlobalConfigurationData Tests" />
   <Property Name="TestNames" Value="$(TestNames);GlobalConfigurationData" />
   <Extend Name="RunAllTests">
-    <Expand Name="DeployAndTest" Platform="Win64" TestName="GlobalConfigurationData" ShortName="GlobalConfigurationData" TargetName="GlobalConfigurationDataTests" BinaryRelativePath="Engine\Binaries" ReportType="console" />
+    <Expand Name="DeployAndTest" Platform="Win64" TestName="GlobalConfigurationData" ShortName="GlobalConfigurationData" TargetName="GlobalConfigurationDataTests" BinaryRelativePath="Engine\Binaries" ReportType="testdata" />
   </Extend>
 </BuildGraph>
\ No newline at end of file
diff --git a/Engine/Build/LowLevelTests/MathCore.xml b/Engine/Build/LowLevelTests/MathCore.xml
--- a/Engine/Build/LowLevelTests/MathCore.xml
+++ b/Engine/Build/LowLevelTests/MathCore.xml
@@ -3,10 +3,10 @@
   <Option Name="RunMathCoreTests" DefaultValue="" Description="Run MathCore Tests" />
   <Property Name="TestNames" Value="$(TestNames);MathCore" />
   <Extend Name="RunAllTests">
-    <Expand Name="DeployAndTest" Platform="Win64" TestName="MathCore" ShortName="MathCore" TargetName="MathCoreTests" BinaryRelativePath="Engine\Binaries" ReportType="xml" ExtraCompilationArgs="-allmodules" />
-    <Expand Name="DeployAndTest" Platform="Linux" TestName="MathCore" ShortName="MathCore" TargetName="MathCoreTests" BinaryRelativePath="Engine\Binaries" ReportType="xml" ExtraCompilationArgs="-allmodules" />
-    <Expand Name="DeployAndTest" Platform="Mac" TestName="MathCore" ShortName="MathCore" TargetName="MathCoreTests" BinaryRelativePath="Engine\Binaries" ReportType="xml" ExtraCompilationArgs="-allmodules" />
-    <Expand Name="DeployAndTest" Platform="Android" TestName="MathCore" ShortName="MathCore" TargetName="MathCoreTests" BinaryRelativePath="Engine\Binaries" ReportType="xml" ExtraCompilationArgs="-allmodules -architectures=arm64" />
-    <Expand Name="DeployAndTest" Platform="IOS" TestName="MathCore" ShortName="MathCore" TargetName="MathCoreTests" BinaryRelativePath="Engine\Binaries" ReportType="xml" ExtraCompilationArgs="-allmodules" RunUnsupported="True" />
+    <Expand Name="DeployAndTest" Platform="Win64" TestName="MathCore" ShortName="MathCore" TargetName="MathCoreTests" BinaryRelativePath="Engine\Binaries" ReportType="testdata" ExtraCompilationArgs="-allmodules" />
+    <Expand Name="DeployAndTest" Platform="Linux" TestName="MathCore" ShortName="MathCore" TargetName="MathCoreTests" BinaryRelativePath="Engine\Binaries" ReportType="testdata" ExtraCompilationArgs="-allmodules" />
+    <Expand Name="DeployAndTest" Platform="Mac" TestName="MathCore" ShortName="MathCore" TargetName="MathCoreTests" BinaryRelativePath="Engine\Binaries" ReportType="testdata" ExtraCompilationArgs="-allmodules" />
+    <Expand Name="DeployAndTest" Platform="Android" TestName="MathCore" ShortName="MathCore" TargetName="MathCoreTests" BinaryRelativePath="Engine\Binaries" ReportType="console" ExtraCompilationArgs="-allmodules -architectures=arm64" />
+    <Expand Name="DeployAndTest" Platform="IOS" TestName="MathCore" ShortName="MathCore" TargetName="MathCoreTests" BinaryRelativePath="Engine\Binaries" ReportType="testdata" ExtraCompilationArgs="-allmodules" RunUnsupported="True" />
   </Extend>
 </BuildGraph>
\ No newline at end of file
diff --git a/Engine/Build/LowLevelTests/NetworkPredictionPlugin.xml b/Engine/Build/LowLevelTests/NetworkPredictionPlugin.xml
--- a/Engine/Build/LowLevelTests/NetworkPredictionPlugin.xml
+++ b/Engine/Build/LowLevelTests/NetworkPredictionPlugin.xml
@@ -3,7 +3,7 @@
   <Option Name="RunNetworkPredictionPluginTests" DefaultValue="" Description="Run Net Prediction Tests" />
   <Property Name="TestNames" Value="$(TestNames);NetworkPredictionPlugin" />
   <Extend Name="RunAllTests">
-    <Expand Name="DeployAndTest" Platform="Win64" TestName="NetworkPredictionPlugin" ShortName="Net Prediction" TargetName="NetworkPredictionTests" BinaryRelativePath="Engine\Binaries" ReportType="console" />
-    <Expand Name="DeployAndTest" Platform="Linux" TestName="NetworkPredictionPlugin" ShortName="Net Prediction" TargetName="NetworkPredictionTests" BinaryRelativePath="Engine\Binaries" ReportType="console" />
+    <Expand Name="DeployAndTest" Platform="Win64" TestName="NetworkPredictionPlugin" ShortName="Net Prediction" TargetName="NetworkPredictionTests" BinaryRelativePath="Engine\Binaries" ReportType="testdata" />
+    <Expand Name="DeployAndTest" Platform="Linux" TestName="NetworkPredictionPlugin" ShortName="Net Prediction" TargetName="NetworkPredictionTests" BinaryRelativePath="Engine\Binaries" ReportType="testdata" />
   </Extend>
 </BuildGraph>
\ No newline at end of file
diff --git a/Engine/Build/LowLevelTests/OnlineServicesNull.xml b/Engine/Build/LowLevelTests/OnlineServicesNull.xml
--- a/Engine/Build/LowLevelTests/OnlineServicesNull.xml
+++ b/Engine/Build/LowLevelTests/OnlineServicesNull.xml
@@ -6,8 +6,8 @@
   <Option Name="RunOnlineServicesNullTests" DefaultValue="" Description="Run Online Services Null Tests" />
   <Property Name="TestNames" Value="$(TestNames);OnlineServicesNull" />
   <Extend Name="RunAllTests">
-    <Expand Name="DeployAndTest" Platform="Win64" TestName="OnlineServicesNull" ShortName="OnlineServicesNull" StagesWithProjectFile="True" TargetName="OnlineServicesNullTests" BinaryRelativePath="Engine\Binaries" ReportType="xml" ExtraArgs="-dumpiniloads" Tags="~[StringProcessing]" />
-    <Expand Name="DeployAndTest" Platform="Linux" TestName="OnlineServicesNull" ShortName="OnlineServicesNull" StagesWithProjectFile="True" TargetName="OnlineServicesNullTests" BinaryRelativePath="Engine\Binaries" ReportType="xml" ExtraArgs="-dumpiniloads" Tags="~[StringProcessing]" />
+    <Expand Name="DeployAndTest" Platform="Win64" TestName="OnlineServicesNull" ShortName="OnlineServicesNull" StagesWithProjectFile="True" TargetName="OnlineServicesNullTests" BinaryRelativePath="Engine\Binaries" ReportType="testdata" ExtraArgs="-dumpiniloads" Tags="~[StringProcessing]" />
+    <Expand Name="DeployAndTest" Platform="Linux" TestName="OnlineServicesNull" ShortName="OnlineServicesNull" StagesWithProjectFile="True" TargetName="OnlineServicesNullTests" BinaryRelativePath="Engine\Binaries" ReportType="testdata" ExtraArgs="-dumpiniloads" Tags="~[StringProcessing]" />
     <!-- <Expand Name="DeployAndTest" Platform="Mac" TestName="OnlineServicesNull" ShortName="OnlineServicesNull" StagesWithProjectFile="True" TargetName="OnlineServicesNullTests" BinaryRelativePath="Engine\Binaries" ReportType="xml" ExtraArgs="-dumpiniloads" Tags="~[StringProcessing]" ExtraCompilationArgs="-ProgramArchitecture=x64" /> -->
     <!-- Currently needs test framework support, leave out from running on Horde for now. -->
     <!-- <Expand Name="DeployAndTest" Platform="IOS" TestName="OnlineServicesNull" ShortName="OnlineServicesNull" StagesWithProjectFile="True" TargetName="OnlineServicesNullTests" BinaryRelativePath="Engine\Binaries" ReportType="xml" PlatformGauntletArgs="~[StringProcessing]" ExtraArgs="-dumpiniloads" Tags="~[StringProcessing]" RunUnsupported="True" /> -->
diff --git a/Engine/Build/LowLevelTests/PlainProps.xml b/Engine/Build/LowLevelTests/PlainProps.xml
--- a/Engine/Build/LowLevelTests/PlainProps.xml
+++ b/Engine/Build/LowLevelTests/PlainProps.xml
@@ -3,6 +3,6 @@
   <Option Name="RunPlainPropsTests" DefaultValue="" Description="Run PlainProps Tests" />
   <Property Name="TestNames" Value="$(TestNames);PlainProps" />
   <Extend Name="RunAllTests">
-    <Expand Name="DeployAndTest" Platform="Win64" Deactivated="True" TestName="PlainProps" ShortName="PlainProps" TargetName="PlainPropsTests" BinaryRelativePath="Engine\Binaries" ReportType="xml" ExtraCompilationArgs="-allmodules" />
+    <Expand Name="DeployAndTest" Platform="Win64" Deactivated="True" TestName="PlainProps" ShortName="PlainProps" TargetName="PlainPropsTests" BinaryRelativePath="Engine\Binaries" ReportType="testdata" ExtraCompilationArgs="-allmodules" />
   </Extend>
 </BuildGraph>
\ No newline at end of file
diff --git a/Engine/Build/LowLevelTests/ReplicationSystem.xml b/Engine/Build/LowLevelTests/ReplicationSystem.xml
--- a/Engine/Build/LowLevelTests/ReplicationSystem.xml
+++ b/Engine/Build/LowLevelTests/ReplicationSystem.xml
@@ -3,8 +3,8 @@
   <Option Name="RunReplicationSystemTests" DefaultValue="" Description="Run Replication System Tests" />
   <Property Name="TestNames" Value="$(TestNames);ReplicationSystem" />
   <Extend Name="RunAllTests">
-    <Expand Name="DeployAndTest" Platform="Win64" TestName="ReplicationSystem" ShortName="Replication System" TargetName="ReplicationSystemLowLevelTests" BinaryRelativePath="Engine\Binaries" ReportType="console" />
-    <Expand Name="DeployAndTest" Platform="Linux" TestName="ReplicationSystem" ShortName="Replication System" TargetName="ReplicationSystemLowLevelTests" BinaryRelativePath="Engine\Binaries" ReportType="console" />
-    <Expand Name="DeployAndTest" Platform="LinuxArm64" TestName="ReplicationSystem" ShortName="Replication System" TargetName="ReplicationSystemLowLevelTests" BinaryRelativePath="Engine\Binaries" ReportType="console" RunContainerized="True" />
+    <Expand Name="DeployAndTest" Platform="Win64" TestName="ReplicationSystem" ShortName="Replication System" TargetName="ReplicationSystemLowLevelTests" BinaryRelativePath="Engine\Binaries" ReportType="testdata" />
+    <Expand Name="DeployAndTest" Platform="Linux" TestName="ReplicationSystem" ShortName="Replication System" TargetName="ReplicationSystemLowLevelTests" BinaryRelativePath="Engine\Binaries" ReportType="testdata" />
+    <Expand Name="DeployAndTest" Platform="LinuxArm64" TestName="ReplicationSystem" ShortName="Replication System" TargetName="ReplicationSystemLowLevelTests" BinaryRelativePath="Engine\Binaries" ReportType="testdata" RunContainerized="True" />
   </Extend>
 </BuildGraph>
\ No newline at end of file
diff --git a/Engine/Build/LowLevelTests/Slate.xml b/Engine/Build/LowLevelTests/Slate.xml
--- a/Engine/Build/LowLevelTests/Slate.xml
+++ b/Engine/Build/LowLevelTests/Slate.xml
@@ -3,8 +3,8 @@
   <Option Name="RunSlateTests" DefaultValue="" Description="Run Slate Tests" />
   <Property Name="TestNames" Value="$(TestNames);Slate" />
   <Extend Name="RunAllTests">
-    <Expand Name="DeployAndTest" Platform="Win64" TestName="Slate" ShortName="Slate" TargetName="SlateTests" BinaryRelativePath="Engine\Binaries" ReportType="xml" ExtraCompilationArgs="-allmodules" />
-    <Expand Name="DeployAndTest" Platform="Linux" TestName="Slate" ShortName="Slate" TargetName="SlateTests" BinaryRelativePath="Engine\Binaries" ReportType="xml" ExtraCompilationArgs="-allmodules" />
-    <Expand Name="DeployAndTest" Platform="Mac" TestName="Slate" ShortName="Slate" TargetName="SlateTests" BinaryRelativePath="Engine\Binaries" ReportType="xml" ExtraCompilationArgs="-allmodules" />
+    <Expand Name="DeployAndTest" Platform="Win64" TestName="Slate" ShortName="Slate" TargetName="SlateTests" BinaryRelativePath="Engine\Binaries" ReportType="testdata" ExtraCompilationArgs="-allmodules" />
+    <Expand Name="DeployAndTest" Platform="Linux" TestName="Slate" ShortName="Slate" TargetName="SlateTests" BinaryRelativePath="Engine\Binaries" ReportType="testdata" ExtraCompilationArgs="-allmodules" />
+    <Expand Name="DeployAndTest" Platform="Mac" TestName="Slate" ShortName="Slate" TargetName="SlateTests" BinaryRelativePath="Engine\Binaries" ReportType="testdata" ExtraCompilationArgs="-allmodules" />
   </Extend>
-</BuildGraph>
+</BuildGraph>
\ No newline at end of file
diff --git a/Engine/Build/LowLevelTests/StateGraph.xml b/Engine/Build/LowLevelTests/StateGraph.xml
--- a/Engine/Build/LowLevelTests/StateGraph.xml
+++ b/Engine/Build/LowLevelTests/StateGraph.xml
@@ -3,6 +3,6 @@
   <Option Name="RunStateGraphTests" DefaultValue="" Description="Run State Graph Tests" />
   <Property Name="TestNames" Value="$(TestNames);StateGraph" />
   <Extend Name="RunAllTests">
-    <Expand Name="DeployAndTest" Platform="Win64" TestName="StateGraph" ShortName="State Graph" TargetName="StateGraphTests" BinaryRelativePath="Engine\Binaries" ReportType="console" />
+    <Expand Name="DeployAndTest" Platform="Win64" TestName="StateGraph" ShortName="State Graph" TargetName="StateGraphTests" BinaryRelativePath="Engine\Binaries" ReportType="testdata" />
   </Extend>
 </BuildGraph>
\ No newline at end of file
diff --git a/Engine/Build/LowLevelTests/SwitchboardListener.xml b/Engine/Build/LowLevelTests/SwitchboardListener.xml
--- a/Engine/Build/LowLevelTests/SwitchboardListener.xml
+++ b/Engine/Build/LowLevelTests/SwitchboardListener.xml
@@ -3,6 +3,6 @@
   <Option Name="RunSwitchboardListenerTests" DefaultValue="" Description="Run SwitchboardListener Tests" />
   <Property Name="TestNames" Value="$(TestNames);SwitchboardListener" />
   <Extend Name="RunAllTests">
-    <Expand Name="DeployAndTest" Platform="Win64" TestName="SwitchboardListener" ShortName="SwitchboardListener" TargetName="SwitchboardListenerTests" BinaryRelativePath="Engine\Binaries" ReportType="console" />
+    <Expand Name="DeployAndTest" Platform="Win64" TestName="SwitchboardListener" ShortName="SwitchboardListener" TargetName="SwitchboardListenerTests" BinaryRelativePath="Engine\Binaries" ReportType="testdata" />
   </Extend>
 </BuildGraph>
\ No newline at end of file
diff --git a/Engine/Build/LowLevelTests/ToolMenus.xml b/Engine/Build/LowLevelTests/ToolMenus.xml
--- a/Engine/Build/LowLevelTests/ToolMenus.xml
+++ b/Engine/Build/LowLevelTests/ToolMenus.xml
@@ -3,8 +3,8 @@
   <Option Name="RunToolMenusTests" DefaultValue="" Description="Run ToolMenus Tests" />
   <Property Name="TestNames" Value="$(TestNames);ToolMenus" />
   <Extend Name="RunAllTests">
-    <Expand Name="DeployAndTest" Platform="Win64" TestName="ToolMenus" ShortName="ToolMenus" TargetName="ToolMenusTests" BinaryRelativePath="Engine\Binaries" ReportType="xml" ExtraCompilationArgs="-allmodules" />
-    <Expand Name="DeployAndTest" Platform="Linux" TestName="ToolMenus" ShortName="ToolMenus" TargetName="ToolMenusTests" BinaryRelativePath="Engine\Binaries" ReportType="xml" ExtraCompilationArgs="-allmodules" />
-    <Expand Name="DeployAndTest" Platform="Mac" TestName="ToolMenus" ShortName="ToolMenus" TargetName="ToolMenusTests" BinaryRelativePath="Engine\Binaries" ReportType="xml" ExtraCompilationArgs="-allmodules" />
+    <Expand Name="DeployAndTest" Platform="Win64" TestName="ToolMenus" ShortName="ToolMenus" TargetName="ToolMenusTests" BinaryRelativePath="Engine\Binaries" ReportType="testdata" ExtraCompilationArgs="-allmodules" />
+    <Expand Name="DeployAndTest" Platform="Linux" TestName="ToolMenus" ShortName="ToolMenus" TargetName="ToolMenusTests" BinaryRelativePath="Engine\Binaries" ReportType="testdata" ExtraCompilationArgs="-allmodules" />
+    <Expand Name="DeployAndTest" Platform="Mac" TestName="ToolMenus" ShortName="ToolMenus" TargetName="ToolMenusTests" BinaryRelativePath="Engine\Binaries" ReportType="testdata" ExtraCompilationArgs="-allmodules" />
   </Extend>
 </BuildGraph>
\ No newline at end of file
diff --git a/Engine/Build/LowLevelTests/Web.xml b/Engine/Build/LowLevelTests/Web.xml
--- a/Engine/Build/LowLevelTests/Web.xml
+++ b/Engine/Build/LowLevelTests/Web.xml
@@ -3,10 +3,10 @@
   <Option Name="RunWebTests" DefaultValue="" Description="Run Web Tests" />
   <Property Name="TestNames" Value="$(TestNames);Web" />
   <Extend Name="RunAllTests">
-    <Expand Name="DeployAndTest" Platform="Win64" TestName="Web" ShortName="Web" TargetName="WebTests" BinaryRelativePath="Engine\Binaries" ReportType="xml" ExtraCompilationArgs="-allmodules" />
-    <Expand Name="DeployAndTest" Platform="Linux" TestName="Web" ShortName="Web" TargetName="WebTests" BinaryRelativePath="Engine\Binaries" ReportType="xml" ExtraCompilationArgs="-allmodules" />
-    <Expand Name="DeployAndTest" Platform="Mac" TestName="Web" ShortName="Web" TargetName="WebTests" BinaryRelativePath="Engine\Binaries" ReportType="xml" ExtraCompilationArgs="-allmodules" />
-    <Expand Name="DeployAndTest" Platform="Android" TestName="Web" ShortName="Web" TargetName="WebTests" BinaryRelativePath="Engine\Binaries" ReportType="xml" ExtraCompilationArgs="-allmodules  -architectures=arm64" RunUnsupported="True" />
-    <Expand Name="DeployAndTest" Platform="IOS" TestName="Web" ShortName="Web" TargetName="WebTests" BinaryRelativePath="Engine\Binaries" ReportType="xml" ExtraCompilationArgs="-allmodules" RunUnsupported="True" />
+    <Expand Name="DeployAndTest" Platform="Win64" TestName="Web" ShortName="Web" TargetName="WebTests" BinaryRelativePath="Engine\Binaries" ReportType="testdata" ExtraCompilationArgs="-allmodules" />
+    <Expand Name="DeployAndTest" Platform="Linux" TestName="Web" ShortName="Web" TargetName="WebTests" BinaryRelativePath="Engine\Binaries" ReportType="testdata" ExtraCompilationArgs="-allmodules" />
+    <Expand Name="DeployAndTest" Platform="Mac" TestName="Web" ShortName="Web" TargetName="WebTests" BinaryRelativePath="Engine\Binaries" ReportType="testdata" ExtraCompilationArgs="-allmodules" />
+    <Expand Name="DeployAndTest" Platform="Android" TestName="Web" ShortName="Web" TargetName="WebTests" BinaryRelativePath="Engine\Binaries" ReportType="console" ExtraCompilationArgs="-allmodules  -architectures=arm64" RunUnsupported="True" />
+    <Expand Name="DeployAndTest" Platform="IOS" TestName="Web" ShortName="Web" TargetName="WebTests" BinaryRelativePath="Engine\Binaries" ReportType="console" ExtraCompilationArgs="-allmodules" RunUnsupported="True" />
   </Extend>
 </BuildGraph>
\ No newline at end of file
diff --git a/Engine/Plugins/Experimental/ChaosRigidPhysicsAsync/Tests/ChaosRigidPhysicsAsyncTests.Build.cs b/Engine/Plugins/Experimental/ChaosRigidPhysicsAsync/Tests/ChaosRigidPhysicsAsyncTests.Build.cs
--- a/Engine/Plugins/Experimental/ChaosRigidPhysicsAsync/Tests/ChaosRigidPhysicsAsyncTests.Build.cs
+++ b/Engine/Plugins/Experimental/ChaosRigidPhysicsAsync/Tests/ChaosRigidPhysicsAsyncTests.Build.cs
@@ -11,7 +11,7 @@ public class ChaosRigidPhysicsAsyncTests : TestModuleRules
 			TestMetadata = new Metadata();
 			TestMetadata.TestName = "ChaosRigidPhysicsAsyncTests";
 			TestMetadata.TestShortName = "ChaosRigidPhysicsAsyncTests";
-			TestMetadata.ReportType = "xml";
+			TestMetadata.ReportType = "testdata";
 			TestMetadata.SupportedPlatforms.Add(UnrealTargetPlatform.Linux);
 			TestMetadata.SupportedPlatforms.Add(UnrealTargetPlatform.Mac);
 			TestMetadata.SupportedPlatforms.Add(UnrealTargetPlatform.Android);
diff --git a/Engine/Source/Developer/LowLevelTestsRunner/Private/TestDataReporter.cpp b/Engine/Source/Developer/LowLevelTestsRunner/Private/TestDataReporter.cpp
new file mode 100644
--- /dev/null
+++ b/Engine/Source/Developer/LowLevelTestsRunner/Private/TestDataReporter.cpp
@@ -0,0 +1,431 @@
+// Copyright Epic Games, Inc. All Rights Reserved.
+
+#include "HAL/FileManager.h"
+#include "HAL/PlatformMisc.h"
+#include "HAL/PlatformProcess.h"
+#include "Misc/App.h"
+#include "Misc/DateTime.h"
+#include "Misc/FileHelper.h"
+#include "Misc/Paths.h"
+#include "Misc/StringBuilder.h"
+
+#include <catch2/catch_test_case_info.hpp>
+#include <catch2/catch_timer.hpp>
+#include <catch2/reporters/catch_reporter_registrars.hpp>
+#include <catch2/reporters/catch_reporter_streaming_base.hpp>
+
+#include <catch2/internal/catch_jsonwriter.hpp>
+#include <catch2/internal/catch_stringref.hpp>
+
+#include <sstream>
+#include <stack>
+#include <string>
+
+using namespace Catch;
+
+namespace UE::LowLevelTests
+{
+	class FTestDataReporter : public StreamingReporterBase
+	{
+	private:
+		std::string m_ReportFilePath;
+	public:
+		using StreamingReporterBase::StreamingReporterBase;
+
+		FTestDataReporter(ReporterConfig&& Config)
+			: StreamingReporterBase(CATCH_MOVE(Config))
+		{
+			m_preferences.shouldRedirectStdOut = true;
+			m_preferences.shouldReportAllAssertions = true;
+			m_preferences.shouldReportAllAssertionStarts = false;
+
+			// Initialize JSON writer with the stringstream
+			m_objectWriters.emplace(m_stream);
+			m_writers.emplace(Writer::Object);
+		}
+
+		~FTestDataReporter() {
+			m_stream << std::endl << std::flush;
+			assert(m_writers.empty());
+		}
+
+		static std::string getDescription()
+		{
+			return "Low Level Tests TestData JSON Reporter";
+		}
+
+		void testRunStarting(TestRunInfo const& TestInfo) override
+		{
+			StreamingReporterBase::testRunStarting(TestInfo);
+
+			Succeeded = 0;
+			SucceededWithWarnings = 0;
+			Failed = 0;
+			NotRun = 0;
+			InProcess = 0;
+			TotalDuration = 0.0f;
+
+			ReportCreatedOn = FDateTime::UtcNow();
+
+			FString OSMajorVersionString, OSSubVersionString;
+			FPlatformMisc::GetOSVersions(OSMajorVersionString, OSSubVersionString);
+
+			FString OSVersionString = OSMajorVersionString + TEXT(" ") + OSSubVersionString;
+
+			// Write Devices array
+			startArray("Devices");
+			auto& DeviceObj = startObject();
+			DeviceObj.write("DeviceName").write(TCHAR_TO_UTF8(*FString(FPlatformProcess::ComputerName())));
+			DeviceObj.write("InstanceName").write(TCHAR_TO_UTF8(*FApp::GetInstanceName()));
+			DeviceObj.write("Platform").write(TCHAR_TO_UTF8(*FString(FPlatformProperties::PlatformName())));
+			DeviceObj.write("SessionId").write(TCHAR_TO_UTF8(*LexToString(FApp::GetSessionId())));
+			DeviceObj.write("OSVersion").write(TCHAR_TO_UTF8(*OSVersionString));
+			DeviceObj.write("Mode").write(TCHAR_TO_UTF8(*FString(FPlatformMisc::GetDefaultDeviceProfileName())));
+			DeviceObj.write("CPU").write(TCHAR_TO_UTF8(*FPlatformMisc::GetCPUBrand().TrimStart()));
+#if !PLATFORM_ANDROID
+			DeviceObj.write("GPU").write(TCHAR_TO_UTF8(*FPlatformMisc::GetPrimaryGPUBrand()));
+#endif
+			DeviceObj.write("RAMInGB").write(FPlatformMemory::GetPhysicalGBRam());
+			DeviceObj.write("RHI").write(TCHAR_TO_UTF8(*FApp::GetGraphicsRHI()));
+			endObject();
+			endArray();
+
+			// Write ReportCreatedOn
+			m_objectWriters.top().write("ReportCreatedOn").write(TCHAR_TO_UTF8(*ReportCreatedOn.ToIso8601()));
+
+			// Start Tests array
+			startArray("Tests");
+		}
+
+		void testRunEnded(TestRunStats const& TestRunStats) override
+		{
+			// End Tests array
+			endArray();
+
+			// Write summary statistics
+			auto& rootObj = m_objectWriters.top();
+			rootObj.write("Succeeded").write(Succeeded);
+			rootObj.write("SucceededWithWarnings").write(SucceededWithWarnings);
+			rootObj.write("Failed").write(Failed);
+			rootObj.write("NotRun").write(NotRun);
+			rootObj.write("InProcess").write(InProcess);
+			rootObj.write("TotalDuration").write(TotalDuration);
+			rootObj.write("ComparisonExported").write(false);
+			rootObj.write("ComparisonExportDirectory").write("");
+
+			// End root object
+			endObject();
+		}
+
+		void testCaseStarting(TestCaseInfo const& TestCaseInfo) override
+		{
+			Timer.start();
+
+			CurrentTestErrors = 0;
+			CurrentTestWarnings = 0;
+			CurrentTestStartTime = FDateTime::UtcNow();
+			CurrentTestEntries.Empty();
+
+			// Start test object
+			auto& testObj = startObject();
+
+			// Write TestDisplayName and FullTestPath
+			testObj.write("TestDisplayName").write(TestCaseInfo.name.c_str());
+			testObj.write("FullTestPath").write(TestCaseInfo.name.c_str());
+
+			// Write Tags array
+			auto& tagsArray = startArray("Tags");
+			FString TestTags(TestCaseInfo.tagsAsString().c_str());
+			TArray<FString> TagArray;
+			TestTags.Replace(TEXT("]"), TEXT("[")).ParseIntoArray(TagArray, TEXT("["));
+			for (const FString& TagStr : TagArray)
+			{
+				if (!TagStr.IsEmpty())
+				{
+					tagsArray.write(TCHAR_TO_UTF8(*TagStr));
+				}
+			}
+			endArray();
+
+			InProcess++;
+		}
+
+		void testCaseEnded(TestCaseStats const& TestCaseStats) override
+		{
+			double Duration = Timer.getElapsedSeconds();
+			TotalDuration += static_cast<float>(Duration);
+
+			InProcess--;
+
+			// Determine final state
+			FString State;
+			if (TestCaseStats.totals.assertions.allOk())
+			{
+				if (CurrentTestWarnings > 0)
+				{
+					State = TEXT("Success");
+					SucceededWithWarnings++;
+				}
+				else
+				{
+					State = TEXT("Success");
+					Succeeded++;
+				}
+			}
+			else
+			{
+				State = TEXT("Fail");
+				Failed++;
+			}
+
+			// Write remaining test fields
+			auto& testObj = m_objectWriters.top();
+			testObj.write("State").write(TCHAR_TO_UTF8(*State));
+
+			// Write empty DeviceInstance array
+			startArray("DeviceInstance");
+			endArray();
+
+			testObj.write("Duration").write(static_cast<float>(Duration));
+			testObj.write("DateTime").write(TCHAR_TO_UTF8(*CurrentTestStartTime.ToIso8601()));
+
+			// Write Entries array
+			auto& entriesArray = startArray("Entries");
+			for (const FBufferedEntry& Entry : CurrentTestEntries)
+			{
+				auto& entryObj = startObject();
+
+				// Write Event object
+				{
+					auto& eventObj = startObject("Event");
+					eventObj.write("Type").write(TCHAR_TO_UTF8(*Entry.Type));
+					eventObj.write("Message").write(TCHAR_TO_UTF8(*Entry.Message));
+					eventObj.write("Context").write(TCHAR_TO_UTF8(*Entry.Context));
+					eventObj.write("Artifact").write(TCHAR_TO_UTF8(*Entry.Artifact));
+					endObject();
+				}
+
+				entryObj.write("Filename").write(TCHAR_TO_UTF8(*Entry.Filename));
+				entryObj.write("LineNumber").write(Entry.LineNumber);
+				entryObj.write("Timestamp").write(TCHAR_TO_UTF8(*Entry.Timestamp.ToIso8601()));
+
+				endObject();
+			}
+			endArray();
+
+			testObj.write("Warnings").write(CurrentTestWarnings);
+			testObj.write("Errors").write(CurrentTestErrors);
+
+			// Write empty Artifacts array
+			startArray("Artifacts");
+			endArray();
+
+			// End test object
+			endObject();
+		}
+
+		void skipTest(TestCaseInfo const& TestInfo) override
+		{
+			// Start test object
+			auto& testObj = startObject();
+
+			// Write test fields
+			testObj.write("TestDisplayName").write(TestInfo.name.c_str());
+			testObj.write("FullTestPath").write(TestInfo.name.c_str());
+
+			// Write Tags array
+			{
+				auto& tagsArray = startArray("Tags");
+				FString TestTags(TestInfo.tagsAsString().c_str());
+				TArray<FString> TagArray;
+				TestTags.Replace(TEXT("]"), TEXT("[")).ParseIntoArray(TagArray, TEXT("["));
+				for (const FString& TagStr : TagArray)
+				{
+					if (!TagStr.IsEmpty())
+					{
+						tagsArray.write(TCHAR_TO_UTF8(*TagStr));
+					}
+				}
+				endArray();
+			}
+
+			testObj.write("State").write("Skipped");
+
+			// Write empty DeviceInstance array
+			startArray("DeviceInstance");
+			endArray();
+
+			testObj.write("Duration").write(0.0f);
+			testObj.write("DateTime").write(TCHAR_TO_UTF8(*FDateTime::UtcNow().ToIso8601()));
+
+			// Write empty Entries array
+			startArray("Entries");
+			endArray();
+
+			testObj.write("Warnings").write(0);
+			testObj.write("Errors").write(0);
+
+			// Write empty Artifacts array
+			startArray("Artifacts");
+			endArray();
+
+			// End test object
+			endObject();
+		}
+
+		void assertionEnded(AssertionStats const& Stats) override
+		{
+			AssertionResult const& Result = Stats.assertionResult;
+
+			if (!Result.isOk())
+			{
+				FBufferedEntry Entry;
+
+				switch (Result.getResultType())
+				{
+				case ResultWas::ThrewException:
+				case ResultWas::FatalErrorCondition:
+				case ResultWas::ExplicitFailure:
+				case ResultWas::ExpressionFailed:
+				case ResultWas::DidntThrowException:
+					Entry.Type = TEXT("Error");
+					CurrentTestErrors++;
+					break;
+				case ResultWas::Warning:
+					Entry.Type = TEXT("Warning");
+					CurrentTestWarnings++;
+					break;
+				case ResultWas::ExplicitSkip:
+					Entry.Type = TEXT("Skipped");
+					NotRun++;
+					break;
+				default:
+					Entry.Type = TEXT("Error");
+					CurrentTestErrors++;
+					break;
+				}
+
+				// Build the message
+				TAnsiStringBuilder<1024> MessageBuilder;
+				if (Stats.totals.assertions.total() > 0)
+				{
+					MessageBuilder << "FAILED:\n";
+					if (Result.hasExpression())
+					{
+						MessageBuilder << "  " << Result.getExpressionInMacro().c_str() << "\n";
+					}
+					if (Result.hasExpandedExpression())
+					{
+						MessageBuilder << "with expansion:\n" << Result.getExpandedExpression().c_str() << "\n";
+					}
+				}
+
+				if (!Result.getMessage().empty())
+				{
+					MessageBuilder << Result.getMessage().data() << '\n';
+				}
+
+				for (auto const& Msg : Stats.infoMessages)
+				{
+					if (Msg.type == ResultWas::Info)
+					{
+						MessageBuilder << Msg.message.data() << '\n';
+					}
+				}
+
+				MessageBuilder << "at " << Result.getSourceInfo().file << '(' << (uint64)Result.getSourceInfo().line << ')';
+
+				Entry.Message = FString(UTF8_TO_TCHAR(*MessageBuilder));
+				Entry.Context = TEXT("");
+				Entry.Artifact = TEXT("");
+				Entry.Filename = FString(Result.getSourceInfo().file);
+				Entry.LineNumber = static_cast<int32>(Result.getSourceInfo().line);
+				Entry.Timestamp = FDateTime::UtcNow();
+
+				CurrentTestEntries.Add(Entry);
+			}
+		}
+
+	private:
+		enum class Writer {
+			Object,
+			Array
+		};
+
+		JsonArrayWriter& startArray() {
+			m_arrayWriters.emplace(m_arrayWriters.top().writeArray());
+			m_writers.emplace(Writer::Array);
+			return m_arrayWriters.top();
+		}
+		JsonArrayWriter& startArray(StringRef key) {
+			m_arrayWriters.emplace(
+				m_objectWriters.top().write(key).writeArray());
+			m_writers.emplace(Writer::Array);
+			return m_arrayWriters.top();
+		}
+
+		JsonObjectWriter& startObject() {
+			m_objectWriters.emplace(m_arrayWriters.top().writeObject());
+			m_writers.emplace(Writer::Object);
+			return m_objectWriters.top();
+		}
+		JsonObjectWriter& startObject(StringRef key) {
+			m_objectWriters.emplace(
+				m_objectWriters.top().write(key).writeObject());
+			m_writers.emplace(Writer::Object);
+			return m_objectWriters.top();
+		}
+
+		void endObject() {
+			assert(isInside(Writer::Object));
+			m_objectWriters.pop();
+			m_writers.pop();
+		}
+		void endArray() {
+			assert(isInside(Writer::Array));
+			m_arrayWriters.pop();
+			m_writers.pop();
+		}
+
+		bool isInside(Writer writer) {
+			return !m_writers.empty() && m_writers.top() == writer;
+		}
+
+		// Invariant:
+		// When m_writers is not empty and its top element is
+		// - Writer::Object, then m_objectWriters is not be empty
+		// - Writer::Array,  then m_arrayWriters shall not be empty
+		std::stack<JsonObjectWriter> m_objectWriters{};
+		std::stack<JsonArrayWriter> m_arrayWriters{};
+		std::stack<Writer> m_writers{};
+
+		// Buffered entry structure for collecting entries during test execution
+		struct FBufferedEntry
+		{
+			FString Type;
+			FString Message;
+			FString Context;
+			FString Artifact;
+			FString Filename;
+			int32 LineNumber;
+			FDateTime Timestamp;
+		};
+
+		TArray<FBufferedEntry> CurrentTestEntries;
+		int32 CurrentTestErrors;
+		int32 CurrentTestWarnings;
+		FDateTime CurrentTestStartTime;
+		Timer Timer;
+
+		FDateTime ReportCreatedOn;
+		int32 Succeeded;
+		int32 SucceededWithWarnings;
+		int32 Failed;
+		int32 NotRun;
+		int32 InProcess;
+		float TotalDuration;
+	};
+
+	CATCH_REGISTER_REPORTER("testdata", FTestDataReporter)
+
+}
\ No newline at end of file
diff --git a/Engine/Source/Programs/AutomationTool/LowLevelTests/RunLowLevelTests.cs b/Engine/Source/Programs/AutomationTool/LowLevelTests/RunLowLevelTests.cs
--- a/Engine/Source/Programs/AutomationTool/LowLevelTests/RunLowLevelTests.cs
+++ b/Engine/Source/Programs/AutomationTool/LowLevelTests/RunLowLevelTests.cs
@@ -43,7 +43,12 @@ namespace LowLevelTests
 				ContextOptions.Build = Path.Combine(Globals.UnrealRootDir, ContextOptions.Build);
 			}
 
-			return RunTests(ContextOptions);
+			var Result = RunTests(ContextOptions);
+
+			Gauntlet.Horde.GenerateStepSummary();
+			Gauntlet.Horde.GenerateJobSummary();
+
+			return Result;
 		}
 
 		public ExitCode RunTests(LowLevelTestExecutorOptions ContextOptions)
@@ -191,6 +196,9 @@ namespace LowLevelTests
 
 		public Type BuildSourceType { get; protected set; }
 
+		[AutoParam]
+		public string HordeTestDataKey = "";
+
 		[AutoParam(UnrealTargetConfiguration.Development)]
 		public UnrealTargetConfiguration Configuration;
 
diff --git a/Engine/Source/Programs/AutomationTool/LowLevelTests/Tests/Gauntlet.LowLevelTestNode.cs b/Engine/Source/Programs/AutomationTool/LowLevelTests/Tests/Gauntlet.LowLevelTestNode.cs
--- a/Engine/Source/Programs/AutomationTool/LowLevelTests/Tests/Gauntlet.LowLevelTestNode.cs
+++ b/Engine/Source/Programs/AutomationTool/LowLevelTests/Tests/Gauntlet.LowLevelTestNode.cs
@@ -40,7 +40,7 @@ namespace LowLevelTests
 		public string DefaultCommandLine;
 		private string ArtifactPath;
 
-		public override string Name { get { return "LowLevelTest"; } }
+		public override string Name { get { return Context != null ? Context.Options.TestApp : "LowLevelTest"; } }
 
 		public override float MaxDuration { protected set; get; }
 
@@ -263,7 +263,8 @@ namespace LowLevelTests
 					LowLevelTestResult = TestResult.Failed;
 					ExitReason = $"Process exited with exit code {ExitCode}";
 				}
-				else if (ReportCopied.HasValue && !ReportCopied.Value)
+				
+				if (ReportCopied.HasValue && !ReportCopied.Value)
 				{
 					LowLevelTestResult = TestResult.Failed;
 					ExitReason = "Unable to read test report";
@@ -301,9 +302,43 @@ namespace LowLevelTests
 						else
 						{
 							LowLevelTestResult = TestResult.Failed;
-							ExitReason = "Tests failed according to xml report";
+							ExitReason = $"Tests failed according to {ReportType} report";
 						}
 					}
+					else if (ReportType == "testdata")
+					{
+						LowLevelTestsJsonTestDataParser LowLevelTestsJsonTestDataParser = new LowLevelTestsJsonTestDataParser(ReportContents);
+						if (LowLevelTestsJsonTestDataParser.HasPassed())
+						{
+							LowLevelTestResult = TestResult.Passed;
+							ExitReason = "Tests passed";
+						}
+						else
+						{
+							LowLevelTestResult = TestResult.Failed;
+							ExitReason = $"Tests failed according to {ReportType} report";
+						}
+
+						MinimalHordeReportGenerator ReportGenerator = new MinimalHordeReportGenerator(
+							new RoleReportData()
+							{
+								Configuration = Context.RoleContext.Configuration,
+								Platform = Context.RoleContext.Platform,
+								RoleType = "Client",
+								SessionStartTime = SessionStartTime,
+								TargetDevice = TestInstance.Device,
+								InitializationEvents = new System.Collections.Generic.List<UnrealTestEvent>(),
+								ClosureEvents = new System.Collections.Generic.List<UnrealTestEvent>()
+							});
+						ITestReport Report = ReportGenerator.GenerateReport(Context.Options.HordeTestDataKey, this, ReportPath, null);
+						ReportGenerator.SubmitToHorde(Report, Name, string.Empty, Context.Options.HordeTestDataKey);
+					}
+					else
+					{
+						LowLevelTestResult = TestResult.Failed;
+						ExitReason = $"Unknown report type {Context.Options.ReportType}";
+						Log.Error(ExitReason);
+					}
 				}
 				else // ReportPath == null
 				{
@@ -328,10 +363,6 @@ namespace LowLevelTests
 					Log.Info(ExitMessage);
 				}
 			}
-			catch
-			{
-				throw;
-			} 
 			finally 
 			{ 
 				foreach (ILowLevelTestsExtension LowLevelTestsExtension in Context.BuildInfo.LowLevelTestsExtensions)
diff --git a/Engine/Source/Programs/AutomationTool/LowLevelTests/Utility/LowLevelTests.JsonTestDataParser.cs b/Engine/Source/Programs/AutomationTool/LowLevelTests/Utility/LowLevelTests.JsonTestDataParser.cs
new file mode 100644
--- /dev/null
+++ b/Engine/Source/Programs/AutomationTool/LowLevelTests/Utility/LowLevelTests.JsonTestDataParser.cs
@@ -0,0 +1,53 @@
+// Copyright Epic Games, Inc. All Rights Reserved.
+
+using Gauntlet;
+using System;
+using System.Text.Json;
+
+namespace LowLevelTests
+{
+	public class LowLevelTestsJsonTestDataParser
+	{
+		private JsonDocument ReportDoc;
+		public bool IsValid { get; protected set; }
+
+		public LowLevelTestsJsonTestDataParser(string InContents)
+		{
+			try
+			{
+				ReportDoc = JsonDocument.Parse(InContents);
+				IsValid = true;
+			}
+			catch (Exception ParseEx)
+			{
+				Log.Warning("Encountered error while parsing report {0}", ParseEx.ToString());
+				IsValid = false;
+			}
+		}
+
+		public bool HasPassed()
+		{
+			if (!IsValid)
+			{
+				return false;
+			}
+
+			try
+			{
+				JsonElement root = ReportDoc.RootElement;
+
+				// Read the summary statistics from the JSON report
+				int succeeded = root.GetProperty("Succeeded").GetInt32();
+				int failed = root.GetProperty("Failed").GetInt32();
+
+				// Test passes if there are no failures and at least one success
+				return failed == 0 && succeeded > 0;
+			}
+			catch (Exception Ex)
+			{
+				Log.Warning("Encountered error while reading test results from report: {0}", Ex.ToString());
+				return false;
+			}
+		}
+	}
+}
diff --git a/Engine/Source/Programs/AutomationTool/LowLevelTests/Utility/LowLevelTests.LogParser.cs b/Engine/Source/Programs/AutomationTool/LowLevelTests/Utility/LowLevelTests.LogParser.cs
--- a/Engine/Source/Programs/AutomationTool/LowLevelTests/Utility/LowLevelTests.LogParser.cs
+++ b/Engine/Source/Programs/AutomationTool/LowLevelTests/Utility/LowLevelTests.LogParser.cs
@@ -1,6 +1,5 @@
 // Copyright Epic Games, Inc. All Rights Reserved.
 
-using Gauntlet;
 using System;
 using System.Linq;
 using System.Text.RegularExpressions;
diff --git a/Engine/Source/Programs/LowLevelTests/FoundationTests/FoundationTests.Build.cs b/Engine/Source/Programs/LowLevelTests/FoundationTests/FoundationTests.Build.cs
--- a/Engine/Source/Programs/LowLevelTests/FoundationTests/FoundationTests.Build.cs
+++ b/Engine/Source/Programs/LowLevelTests/FoundationTests/FoundationTests.Build.cs
@@ -10,7 +10,7 @@ public class FoundationTests : TestModuleRules
 			TestMetadata = new Metadata();
 			TestMetadata.TestName = "Foundation";
 			TestMetadata.TestShortName = "Foundation";
-			TestMetadata.ReportType = "xml";
+			TestMetadata.ReportType = "testdata";
 			TestMetadata.SupportedPlatforms.Add(UnrealTargetPlatform.Linux);
 			TestMetadata.SupportedPlatforms.Add(UnrealTargetPlatform.Mac);
 			TestMetadata.SupportedPlatforms.Add(UnrealTargetPlatform.Android);
diff --git a/Engine/Source/Programs/LowLevelTests/MathCoreTests/MathCoreTests.Build.cs b/Engine/Source/Programs/LowLevelTests/MathCoreTests/MathCoreTests.Build.cs
--- a/Engine/Source/Programs/LowLevelTests/MathCoreTests/MathCoreTests.Build.cs
+++ b/Engine/Source/Programs/LowLevelTests/MathCoreTests/MathCoreTests.Build.cs
@@ -10,7 +10,7 @@ public class MathCoreTests : TestModuleRules
 			TestMetadata = new Metadata();
 			TestMetadata.TestName = "MathCore";
 			TestMetadata.TestShortName = "MathCore";
-			TestMetadata.ReportType = "xml";
+			TestMetadata.ReportType = "testdata";
 			TestMetadata.SupportedPlatforms.Add(UnrealTargetPlatform.Linux);
 			TestMetadata.SupportedPlatforms.Add(UnrealTargetPlatform.Mac);
 			TestMetadata.SupportedPlatforms.Add(UnrealTargetPlatform.Android);
diff --git a/Engine/Source/Programs/Online/OnlineServicesNullTests/OnlineServicesNullTests.Build.cs b/Engine/Source/Programs/Online/OnlineServicesNullTests/OnlineServicesNullTests.Build.cs
--- a/Engine/Source/Programs/Online/OnlineServicesNullTests/OnlineServicesNullTests.Build.cs
+++ b/Engine/Source/Programs/Online/OnlineServicesNullTests/OnlineServicesNullTests.Build.cs
@@ -10,7 +10,7 @@ public class OnlineServicesNullTests : TestModuleRules
 		TestMetadata = new Metadata();
 		TestMetadata.TestName = "OnlineServicesNullTests";
 		TestMetadata.TestShortName = "OnlineServicesNull";
-		TestMetadata.ReportType = "xml";
+		TestMetadata.ReportType = "testdata";
 		TestMetadata.StagesWithProjectFile = true;
 
 		// Ensure LowLevelTestsRunner adds Engine/CoreUObject/ApplicationCore to its PrivateDependencyModuleNames.
diff --git a/Engine/Source/Programs/PlainPropsTests/PlainPropsTests.Build.cs b/Engine/Source/Programs/PlainPropsTests/PlainPropsTests.Build.cs
--- a/Engine/Source/Programs/PlainPropsTests/PlainPropsTests.Build.cs
+++ b/Engine/Source/Programs/PlainPropsTests/PlainPropsTests.Build.cs
@@ -10,7 +10,7 @@ public class PlainPropsTests : TestModuleRules
 			TestMetadata = new Metadata();
 			TestMetadata.TestName = "PlainProps";
 			TestMetadata.TestShortName = "PlainProps";
-			TestMetadata.ReportType = "xml";
+			TestMetadata.ReportType = "testdata";
 			TestMetadata.Deactivated = true;
 			//TestMetadata.SupportedPlatforms.Add(UnrealTargetPlatform.Linux);
 			//TestMetadata.SupportedPlatforms.Add(UnrealTargetPlatform.Mac);
diff --git a/Engine/Source/Programs/ReplicationSystemLowLevelTests/ReplicationSystemLowLevelTests.Build.cs b/Engine/Source/Programs/ReplicationSystemLowLevelTests/ReplicationSystemLowLevelTests.Build.cs
--- a/Engine/Source/Programs/ReplicationSystemLowLevelTests/ReplicationSystemLowLevelTests.Build.cs
+++ b/Engine/Source/Programs/ReplicationSystemLowLevelTests/ReplicationSystemLowLevelTests.Build.cs
@@ -11,6 +11,7 @@ public class ReplicationSystemLowLevelTests : TestModuleRules
 			TestMetadata = new Metadata();
 			TestMetadata.TestName = "ReplicationSystem";
 			TestMetadata.TestShortName = "Replication System";
+			TestMetadata.ReportType = "testdata";
 			TestMetadata.SupportedPlatforms.Add(UnrealTargetPlatform.Linux);
 			TestMetadata.SupportedPlatforms.Add(UnrealTargetPlatform.LinuxArm64);
 			TestMetadata.PlatformRunContainerized.Add(UnrealTargetPlatform.LinuxArm64, true);
diff --git a/Engine/Source/Programs/SlateTests/SlateTests.Build.cs b/Engine/Source/Programs/SlateTests/SlateTests.Build.cs
--- a/Engine/Source/Programs/SlateTests/SlateTests.Build.cs
+++ b/Engine/Source/Programs/SlateTests/SlateTests.Build.cs
@@ -12,7 +12,7 @@ public class SlateTests : TestModuleRules
 			TestMetadata = new Metadata();
 			TestMetadata.TestName = "Slate";
 			TestMetadata.TestShortName = "Slate";
-			TestMetadata.ReportType = "xml";
+			TestMetadata.ReportType = "testdata";
 			TestMetadata.SupportedPlatforms.Add(UnrealTargetPlatform.Linux);
 			TestMetadata.SupportedPlatforms.Add(UnrealTargetPlatform.Mac);
 
diff --git a/Engine/Source/Programs/ToolMenusTests/ToolMenusTests.Build.cs b/Engine/Source/Programs/ToolMenusTests/ToolMenusTests.Build.cs
--- a/Engine/Source/Programs/ToolMenusTests/ToolMenusTests.Build.cs
+++ b/Engine/Source/Programs/ToolMenusTests/ToolMenusTests.Build.cs
@@ -10,7 +10,7 @@ public class ToolMenusTests : TestModuleRules
 			TestMetadata = new Metadata();
 			TestMetadata.TestName = "ToolMenus";
 			TestMetadata.TestShortName = "ToolMenus";
-			TestMetadata.ReportType = "xml";
+			TestMetadata.ReportType = "testdata";
 			TestMetadata.SupportedPlatforms.Add(UnrealTargetPlatform.Linux);
 			TestMetadata.SupportedPlatforms.Add(UnrealTargetPlatform.Mac);
 
diff --git a/Engine/Source/Programs/UnrealBuildTool/Configuration/Rules/TestModuleRules.cs b/Engine/Source/Programs/UnrealBuildTool/Configuration/Rules/TestModuleRules.cs
--- a/Engine/Source/Programs/UnrealBuildTool/Configuration/Rules/TestModuleRules.cs
+++ b/Engine/Source/Programs/UnrealBuildTool/Configuration/Rules/TestModuleRules.cs
@@ -487,7 +487,8 @@ namespace UnrealBuildTool
 			/// </summary>
 			public string TestShortName { get; set; }
 
-			private string ReportTypePrivate = "console";
+			// Default to JSON test data output
+			private string ReportTypePrivate = "testdata";
 			/// <summary>
 			/// Type of Catch2 report, defaults to console.
 			/// </summary>
diff --git a/Engine/Source/Programs/WebTests/WebTests.Build.cs b/Engine/Source/Programs/WebTests/WebTests.Build.cs
--- a/Engine/Source/Programs/WebTests/WebTests.Build.cs
+++ b/Engine/Source/Programs/WebTests/WebTests.Build.cs
@@ -11,7 +11,7 @@ public class WebTests : TestModuleRules
 			TestMetadata = new Metadata();
 			TestMetadata.TestName = "Web";
 			TestMetadata.TestShortName = "Web";
-			TestMetadata.ReportType = "xml";
+			TestMetadata.ReportType = "testdata";
 			TestMetadata.SupportedPlatforms.Add(UnrealTargetPlatform.Linux);
 			TestMetadata.SupportedPlatforms.Add(UnrealTargetPlatform.Mac);
 			TestMetadata.SupportedPlatforms.Add(UnrealTargetPlatform.Android);
